Fix PropertyPathHelpers bitfield automated tests

Bitfield tests were passing or failing intermittently under VS2015 because of uninitialized variables in the test code. VS2017 was failing consistently, as was PS4.
Fixed bitfield tests by supplying partial specializations of helper functions  for bool/bitfield properties.
Also added a test for copying between two bitfields.

diff --git a/Engine/Source/Runtime/PropertyPath/Private/PropertyPathHelpers.cpp b/Engine/Source/Runtime/PropertyPath/Private/PropertyPathHelpers.cpp
--- a/Engine/Source/Runtime/PropertyPath/Private/PropertyPathHelpers.cpp
+++ b/Engine/Source/Runtime/PropertyPath/Private/PropertyPathHelpers.cpp
@@ -409,6 +409,12 @@ namespace PropertyPathHelpersInternal
 			{
 				DestProperty->CopyCompleteValue(InDestPropertyPath.GetCachedAddress(), InSrcPropertyPath.GetCachedAddress());
 			}
+			else if(UBoolProperty* DestBoolProperty = Cast<UBoolProperty>(DestProperty))
+			{
+				UBoolProperty* SrcBoolProperty = Cast<UBoolProperty>(InSrcPropertyPath.GetLastSegment().GetField());
+				const bool bValue = SrcBoolProperty->GetPropertyValue(InSrcPropertyPath.GetCachedAddress());
+				DestBoolProperty->SetPropertyValue(InDestPropertyPath.GetCachedAddress(), bValue);
+			}
 			else
 			{
 				DestProperty->CopySingleValue(InDestPropertyPath.GetCachedAddress(), InSrcPropertyPath.GetCachedAddress());
diff --git a/Engine/Source/Runtime/PropertyPath/Public/PropertyPathHelpers.h b/Engine/Source/Runtime/PropertyPath/Public/PropertyPathHelpers.h
--- a/Engine/Source/Runtime/PropertyPath/Public/PropertyPathHelpers.h
+++ b/Engine/Source/Runtime/PropertyPath/Public/PropertyPathHelpers.h
@@ -667,6 +667,37 @@ namespace PropertyPathHelpersInternal
 		}
 	};
 
+	/** Partial specialization for bools/bitfields */
+	template<typename ContainerType>
+	struct FGetValueHelper<bool, ContainerType>
+	{
+		static bool GetValue(ContainerType* InContainer, const FCachedPropertyPath& InPropertyPath, bool& OutValue, UProperty*& OutProperty)
+		{
+			const FPropertyPathSegment& LastSegment = InPropertyPath.GetLastSegment();
+			int32 ArrayIndex = LastSegment.GetArrayIndex();
+			UProperty* Property = CastChecked<UProperty>(LastSegment.GetField());
+
+			// Verify that the cpp type matches a known property type.
+			if ( IsConcreteTypeCompatibleWithReflectedType<bool>(Property) )
+			{
+				ArrayIndex = ArrayIndex == INDEX_NONE ? 0 : ArrayIndex;
+				if ( PropertySizesMatch<bool>(Property) && ArrayIndex < Property->ArrayDim )
+				{
+					if(void* Address = Property->ContainerPtrToValuePtr<bool>(InContainer, ArrayIndex))
+					{
+						InPropertyPath.ResolveLeaf(Address);
+						UBoolProperty* BoolProperty = CastChecked<UBoolProperty>(LastSegment.GetField());
+						OutValue = BoolProperty->GetPropertyValue(Address);
+						OutProperty = Property;
+						return true;
+					}
+				}
+			}
+
+			return false;
+		}
+	};
+
 	/** 
 	 * Resolve a property path to a property and a value. Supports functions as input fields.
 	 * @param InContainer			The containing object/structure for the current path iteration
@@ -854,7 +885,38 @@ namespace PropertyPathHelpersInternal
 				}
 			}
 
-			return true;
+			return false;
+		}
+	};
+
+	/** Partial specialization for bools/bitfields */
+	template<typename ContainerType>
+	struct FSetValueHelper<bool, ContainerType>
+	{
+		static bool SetValue(ContainerType* InContainer, const FCachedPropertyPath& InPropertyPath, const bool& InValue)
+		{
+			const FPropertyPathSegment& LastSegment = InPropertyPath.GetLastSegment();
+			int32 ArrayIndex = LastSegment.GetArrayIndex();
+			UProperty* Property = CastChecked<UProperty>(LastSegment.GetField());
+
+			// Verify that the cpp type matches a known property type.
+			if ( IsConcreteTypeCompatibleWithReflectedType<bool>(Property) )
+			{
+				// Ensure that the element sizes are the same, prevents the user from doing something terribly wrong.
+				ArrayIndex = ArrayIndex == INDEX_NONE ? 0 : ArrayIndex;
+				if ( PropertySizesMatch<bool>(Property) && ArrayIndex < Property->ArrayDim )
+				{
+					if(void* Address = Property->ContainerPtrToValuePtr<bool>(InContainer, ArrayIndex))
+					{
+						InPropertyPath.ResolveLeaf(Address);
+						UBoolProperty* BoolProperty = CastChecked<UBoolProperty>(LastSegment.GetField());
+						BoolProperty->SetPropertyValue(Address, InValue);
+						return true;
+					}
+				}
+			}
+
+			return false;
 		}
 	};
 
@@ -1020,6 +1082,28 @@ namespace PropertyPathHelpersInternal
 		}
 	};
 
+	/** Explicit specialization for bools/bitfields */
+	template<>
+	struct FGetValueFastHelper<bool>
+	{
+		static bool GetValue(const FCachedPropertyPath& InPropertyPath, bool& OutValue, UProperty*& OutProperty)
+		{
+			const FPropertyPathSegment& LastSegment = InPropertyPath.GetLastSegment();
+			OutProperty = CastChecked<UProperty>(LastSegment.GetField());
+			UArrayProperty* ArrayProp = Cast<UArrayProperty>(OutProperty);
+			if ( ArrayProp && LastSegment.GetArrayIndex() != INDEX_NONE )
+			{
+				ArrayProp->Inner->CopySingleValue(&OutValue, InPropertyPath.GetCachedAddress());
+			}
+			else
+			{
+				UBoolProperty* BoolProperty = CastChecked<UBoolProperty>(OutProperty);
+				OutValue = BoolProperty->GetPropertyValue(InPropertyPath.GetCachedAddress());
+			}
+			return true;
+		}
+	};
+
 	/** 
 	 * Fast, unsafe version of GetValue().
 	 * @param	InContainer		The containing object/structure to iterate against
@@ -1077,6 +1161,29 @@ namespace PropertyPathHelpersInternal
 		}
 	};
 
+	/** Explicit specialization for bools/bitfields */
+	template<>
+	struct FSetValueFastHelper<bool>
+	{
+		static bool SetValue(const FCachedPropertyPath& InPropertyPath, const bool& InValue)
+		{
+			const FPropertyPathSegment& LastSegment = InPropertyPath.GetLastSegment();
+			UProperty* Property = CastChecked<UProperty>(LastSegment.GetField());
+			UArrayProperty* ArrayProp = Cast<UArrayProperty>(Property);
+			if ( ArrayProp && LastSegment.GetArrayIndex() != INDEX_NONE )
+			{
+				ArrayProp->Inner->CopySingleValue(InPropertyPath.GetCachedAddress(), &InValue);
+			}
+			else
+			{
+				UBoolProperty* BoolProperty = CastChecked<UBoolProperty>(Property);
+				BoolProperty->SetPropertyValue(InPropertyPath.GetCachedAddress(), InValue);
+			}
+			
+			return true;
+		}
+	};
+
 	/** 
 	 * Fast, unsafe version of SetValue().
 	 * @param	InContainer		The containing object/structure to iterate against
