Fix to ensure on load. On load some binding may not be fully valid because of other depenmdant classes not being loaded yet. In those cases, we may end up with the SourcePath with property name set to Name_None.

diff --git a/Engine/Source/Runtime/PropertyPath/Private/PropertyPathHelpers.cpp b/Engine/Source/Runtime/PropertyPath/Private/PropertyPathHelpers.cpp
--- a/Engine/Source/Runtime/PropertyPath/Private/PropertyPathHelpers.cpp
+++ b/Engine/Source/Runtime/PropertyPath/Private/PropertyPathHelpers.cpp
@@ -504,7 +504,6 @@ FPropertyPathSegment::FPropertyPathSegment(int32 InCount, const TCHAR* InString)
 	ensure(PropertyName != nullptr);
 	FString PropertyNameString(PropertyNameLength, PropertyName);
 	Name = FName(*PropertyNameString, FNAME_Find);
-	ensure(Name != NAME_None);
 }
 
 FPropertyPathSegment FPropertyPathSegment::MakeUnresolvedCopy(const FPropertyPathSegment& ToCopy)
