PropertyPathHelpers - Making several improvements to the property path evaluators.  They can now evaluate multiple UFunctions in a chain, better support for some types like FName, and some utility ctors for handling singular segments.

diff --git a/Engine/Source/Runtime/PropertyPath/Private/PropertyPathHelpers.cpp b/Engine/Source/Runtime/PropertyPath/Private/PropertyPathHelpers.cpp
--- a/Engine/Source/Runtime/PropertyPath/Private/PropertyPathHelpers.cpp
+++ b/Engine/Source/Runtime/PropertyPath/Private/PropertyPathHelpers.cpp
@@ -24,9 +24,15 @@ namespace PropertyPathHelpersInternal
 		// Obtain the property info from the given structure definition
 		if ( UField* Field = Segment.Resolve(InStruct) )
 		{
+			const bool bFinalSegment = SegmentIndex == (InPropertyPath.GetNumSegments() - 1);
+
 			if ( UProperty* Property = Cast<UProperty>(Field) )
 			{
-				if ( SegmentIndex < ( InPropertyPath.GetNumSegments() - 1 ) )
+				if (bFinalSegment)
+				{
+					return InResolver.Resolve(static_cast<ContainerType*>(InContainer), InPropertyPath);
+				}
+				else
 				{
 					// Check first to see if this is a simple object (eg. not an array of objects)
 					if ( UObjectProperty* ObjectProperty = Cast<UObjectProperty>(Property) )
@@ -107,18 +113,35 @@ namespace PropertyPathHelpersInternal
 						// TODO: we dont support map properties yet
 					}
 				}
-				else
-				{
-					return InResolver.Resolve(static_cast<ContainerType*>(InContainer), InPropertyPath);
-				}
 			}
 			else
 			{
-				// Only allow functions as the final link in the chain.
-				if ( SegmentIndex == ( InPropertyPath.GetNumSegments() - 1 ) )
+				// If it's the final segment, use the resolver to get the value.
+				if (bFinalSegment)
 				{
 					return InResolver.Resolve(static_cast<ContainerType*>(InContainer), InPropertyPath);
 				}
+				else
+				{
+					// If it's not the final segment, but still a function, we're going to treat it as an Object* getter.
+					// in the hopes that it leads to another object that we can resolve the next segment on.  These
+					// getter functions must be very simple.
+
+					UObject* CurrentObject = nullptr;
+					UProperty* GetterProperty = nullptr;
+					FInternalGetterResolver<UObject*> GetterResolver(CurrentObject, GetterProperty);
+
+					FCachedPropertyPath TempPath(Segment);
+					if (GetterResolver.Resolve(InContainer, TempPath))
+					{
+						if (CurrentObject)
+						{
+							InPropertyPath.SetCanSafelyUsedCachedAddress(false);
+
+							return IteratePropertyPathRecursive(CurrentObject->GetClass(), CurrentObject, SegmentIndex + 1, InPropertyPath, InResolver);
+						}
+					}
+				}
 			}
 		}
 
@@ -451,13 +474,23 @@ namespace PropertyPathHelpersInternal
 
 	bool ResolvePropertyPath(UObject* InContainer, const FString& InPropertyPath, FPropertyPathResolver& InResolver)
 	{
-		FCachedPropertyPath InternalPropertyPath(InPropertyPath);
-		return IteratePropertyPathRecursive<UObject>(InContainer->GetClass(), InContainer, 0, InternalPropertyPath, InResolver);
+		if (InContainer)
+		{
+			FCachedPropertyPath InternalPropertyPath(InPropertyPath);
+			return IteratePropertyPathRecursive<UObject>(InContainer->GetClass(), InContainer, 0, InternalPropertyPath, InResolver);
+		}
+
+		return false;
 	}
 
 	bool ResolvePropertyPath(UObject* InContainer, const FCachedPropertyPath& InPropertyPath, FPropertyPathResolver& InResolver)
 	{
-		return IteratePropertyPathRecursive<UObject>(InContainer->GetClass(), InContainer, 0, InPropertyPath, InResolver);
+		if (InContainer)
+		{
+			return IteratePropertyPathRecursive<UObject>(InContainer->GetClass(), InContainer, 0, InPropertyPath, InResolver);
+		}
+
+		return false;
 	}
 
 	bool ResolvePropertyPath(void* InContainer, UStruct* InStruct, const FString& InPropertyPath, FPropertyPathResolver& InResolver)
@@ -573,7 +606,7 @@ FCachedPropertyPath::FCachedPropertyPath(const FString& Path)
 	MakeFromString(Path);
 }
 
-FCachedPropertyPath::FCachedPropertyPath(const TArray<FString>& PropertyChain)
+FCachedPropertyPath::FCachedPropertyPath(const TArray<FString>& PathSegments)
 	: CachedAddress(nullptr)
 	, CachedFunction(nullptr)
 #if DO_CHECK
@@ -581,12 +614,23 @@ FCachedPropertyPath::FCachedPropertyPath(const TArray<FString>& PropertyChain)
 #endif
 	, bCanSafelyUsedCachedAddress(false)
 {
-	for(const FString& Segment : PropertyChain)
+	for (const FString& Segment : PathSegments)
 	{
 		Segments.Add(FPropertyPathSegment(Segment.Len(), *Segment));
 	}
 }
 
+FCachedPropertyPath::FCachedPropertyPath(const FPropertyPathSegment& Segment)
+	: CachedAddress(nullptr)
+	, CachedFunction(nullptr)
+#if DO_CHECK
+	, CachedContainer(nullptr)
+#endif
+	, bCanSafelyUsedCachedAddress(false)
+{
+	Segments.Add(Segment);
+}
+
 void FCachedPropertyPath::MakeFromString(const FString& InPropertyPath)
 {
 	const TCHAR Delim = TEXT('.');
diff --git a/Engine/Source/Runtime/PropertyPath/Public/PropertyPathHelpers.h b/Engine/Source/Runtime/PropertyPath/Public/PropertyPathHelpers.h
--- a/Engine/Source/Runtime/PropertyPath/Public/PropertyPathHelpers.h
+++ b/Engine/Source/Runtime/PropertyPath/Public/PropertyPathHelpers.h
@@ -80,7 +80,10 @@ struct PROPERTYPATH_API FCachedPropertyPath
 	FCachedPropertyPath(const FString& Path);
 
 	/** */
-	FCachedPropertyPath(const TArray<FString>& PropertyChain);
+	FCachedPropertyPath(const FPropertyPathSegment& Segment);
+
+	/** */
+	FCachedPropertyPath(const TArray<FString>& PathSegments);
 
 	/** Check whether this property path is non-empty */
 	bool IsValid() const { return Segments.Num() > 0; }
@@ -91,9 +94,6 @@ struct PROPERTYPATH_API FCachedPropertyPath
 	/** Make a copy which is unresolved */
 	static FCachedPropertyPath MakeUnresolvedCopy(const FCachedPropertyPath& ToCopy);
 
-	/** Make a new property path from a leaf property and an outer class */
-	void MakeFromPropertyAndOuterClass(UProperty* InProperty, UClass* InClass);
-
 	/** @return Get the number of segments in this path */
 	int32 GetNumSegments() const;
 
@@ -291,7 +291,6 @@ namespace PropertyPathHelpers
 	 */
 	PROPERTYPATH_API bool SetPropertyValueFromString(UObject* InContainer, const FCachedPropertyPath& InPropertyPath, const FString& InValue);
 
-
 	/** 
 	 * Set the value represented by this property path from a string 
 	 * @param	InContainer		A pointer to the container structure to resolve the property path against
@@ -839,6 +838,10 @@ namespace PropertyPathHelpersInternal
 					}
 				}
 			}
+			else
+			{
+				//LOG ERROR
+			}
 
 			return false;
 		}
diff --git a/Engine/Source/Runtime/PropertyPath/Public/PropertyTypeCompatibility.h b/Engine/Source/Runtime/PropertyPath/Public/PropertyTypeCompatibility.h
--- a/Engine/Source/Runtime/PropertyPath/Public/PropertyTypeCompatibility.h
+++ b/Engine/Source/Runtime/PropertyPath/Public/PropertyTypeCompatibility.h
@@ -156,6 +156,12 @@ inline bool IsConcreteTypeCompatibleWithReflectedType_Impl<FString>(UProperty* P
 	return Property->GetClass() == UStrProperty::StaticClass();
 }
 
+template<>
+inline bool IsConcreteTypeCompatibleWithReflectedType_Impl<FName>(UProperty* Property)
+{
+	return Property->GetClass() == UNameProperty::StaticClass();
+}
+
 template<typename T>
 inline bool IsConcreteTypeCompatibleWithReflectedType_BuiltInStruct(UProperty* Property)
 {
