FFieldPathSegments will now resolve and cache field pointers if their Struct has been serialized (for example when serializing CDOs)

diff --git a/Engine/Source/Runtime/PropertyPath/Private/PropertyPathHelpers.cpp b/Engine/Source/Runtime/PropertyPath/Private/PropertyPathHelpers.cpp
--- a/Engine/Source/Runtime/PropertyPath/Private/PropertyPathHelpers.cpp
+++ b/Engine/Source/Runtime/PropertyPath/Private/PropertyPathHelpers.cpp
@@ -539,6 +539,17 @@ FPropertyPathSegment::FPropertyPathSegment(int32 InCount, const TCHAR* InString)
 	Name = FName(*PropertyNameString, FNAME_Find);
 }
 
+void FPropertyPathSegment::PostSerialize(const FArchive& Ar)
+{
+	if (Struct)
+	{
+		// if the struct has been serialized then we're loading a CDO and we need to re-cache the field pointer
+		UStruct* ResolveStruct = Struct;
+		Struct = nullptr;
+		Resolve(ResolveStruct);
+	}
+}
+
 FPropertyPathSegment FPropertyPathSegment::MakeUnresolvedCopy(const FPropertyPathSegment& ToCopy)
 {
 	FPropertyPathSegment Segment;
@@ -556,10 +567,10 @@ FFieldVariant FPropertyPathSegment::Resolve(UStruct* InStruct) const
 		if ( InStruct != Struct )
 		{
 			Struct = InStruct;
-			Field = FindField<UField>(InStruct, Name);
+			Field = FindField<FProperty>(InStruct, Name);
 			if (!Field.IsValid())
 			{
-				Field = FindField<FProperty>(InStruct, Name);
+				Field = FindField<UField>(InStruct, Name);
 			}
 		}
 
diff --git a/Engine/Source/Runtime/PropertyPath/Public/PropertyPathHelpers.h b/Engine/Source/Runtime/PropertyPath/Public/PropertyPathHelpers.h
--- a/Engine/Source/Runtime/PropertyPath/Public/PropertyPathHelpers.h
+++ b/Engine/Source/Runtime/PropertyPath/Public/PropertyPathHelpers.h
@@ -22,6 +22,11 @@ struct PROPERTYPATH_API FPropertyPathSegment
 	/** Construct from char count and storage */
 	FPropertyPathSegment(int32 InCount, const TCHAR* InString);
 
+	/**
+	 * Called after this event has been serialized in order to cache the field pointer if necessary
+	 */
+	void PostSerialize(const FArchive& Ar);
+
 	/** Make a copy which is unresolved */
 	static FPropertyPathSegment MakeUnresolvedCopy(const FPropertyPathSegment& ToCopy);
 
@@ -66,6 +71,12 @@ private:
 	mutable FFieldVariant Field;
 };
 
+template<>
+struct TStructOpsTypeTraits<FPropertyPathSegment> : TStructOpsTypeTraitsBase2<FPropertyPathSegment>
+{
+	enum { WithPostSerialize = true };
+};
+
 /** Base class for cached property paths */
 USTRUCT()
 struct PROPERTYPATH_API FCachedPropertyPath
