Fixed PropertyPathHelpers using the wrong version of ExportTextItem when exporting static array value

diff --git a/Engine/Source/Runtime/PropertyPath/Private/PropertyPathHelpers.cpp b/Engine/Source/Runtime/PropertyPath/Private/PropertyPathHelpers.cpp
--- a/Engine/Source/Runtime/PropertyPath/Private/PropertyPathHelpers.cpp
+++ b/Engine/Source/Runtime/PropertyPath/Private/PropertyPathHelpers.cpp
@@ -236,7 +236,7 @@ namespace PropertyPathHelpersInternal
 				if ( void* ValuePtr = Property->ContainerPtrToValuePtr<void>(InContainer, ArrayIndex) )
 				{
 					OutProperty = Property;
-					OutProperty->ExportTextItem_InContainer(OutValue, InContainer, nullptr, nullptr, 0);
+					OutProperty->ExportTextItem_Direct(OutValue, ValuePtr, nullptr, nullptr, 0);
 					return true;
 				}
 			}
