Fix bad memory access when the property type doesn't match the provided type and the path is fully resolved.

diff --git a/Engine/Source/Runtime/PropertyPath/Public/PropertyPathHelpers.h b/Engine/Source/Runtime/PropertyPath/Public/PropertyPathHelpers.h
--- a/Engine/Source/Runtime/PropertyPath/Public/PropertyPathHelpers.h
+++ b/Engine/Source/Runtime/PropertyPath/Public/PropertyPathHelpers.h
@@ -1100,13 +1100,18 @@ namespace PropertyPathHelpersInternal
 			FArrayProperty* ArrayProp = CastField<FArrayProperty>(OutProperty);
 			if ( ArrayProp && LastSegment.GetArrayIndex() != INDEX_NONE )
 			{
-				ArrayProp->Inner->CopySingleValue(&OutValue, InPropertyPath.GetCachedAddress());
+				if (IsConcreteTypeCompatibleWithReflectedType<T>(ArrayProp->Inner))
+				{
+					ArrayProp->Inner->CopySingleValue(&OutValue, InPropertyPath.GetCachedAddress());
+					return true;
+				}
 			}
-			else
+			else if (IsConcreteTypeCompatibleWithReflectedType<T>(OutProperty))
 			{
 				OutProperty->CopySingleValue(&OutValue, InPropertyPath.GetCachedAddress());
+				return true;
 			}
-			return true;
+			return false;
 		}
 	};
 
@@ -1118,8 +1123,12 @@ namespace PropertyPathHelpersInternal
 		{
 			const FPropertyPathSegment& LastSegment = InPropertyPath.GetLastSegment();
 			OutProperty = CastFieldChecked<FProperty>(LastSegment.GetField().ToField());
-			OutProperty->CopyCompleteValue(&OutValue, InPropertyPath.GetCachedAddress());
-			return true;
+			if (IsConcreteTypeCompatibleWithReflectedType<T>(OutProperty))
+			{
+				OutProperty->CopyCompleteValue(&OutValue, InPropertyPath.GetCachedAddress());
+				return true;
+			}
+			return false;
 		}
 	};
 
@@ -1134,14 +1143,19 @@ namespace PropertyPathHelpersInternal
 			FArrayProperty* ArrayProp = CastField<FArrayProperty>(OutProperty);
 			if ( ArrayProp && LastSegment.GetArrayIndex() != INDEX_NONE )
 			{
-				ArrayProp->Inner->CopySingleValue(&OutValue, InPropertyPath.GetCachedAddress());
+				if (IsConcreteTypeCompatibleWithReflectedType<bool>(ArrayProp->Inner))
+				{
+					ArrayProp->Inner->CopySingleValue(&OutValue, InPropertyPath.GetCachedAddress());
+					return true;
+				}
 			}
-			else
+			else if (IsConcreteTypeCompatibleWithReflectedType<bool>(OutProperty))
 			{
 				FBoolProperty* BoolProperty = CastFieldChecked<FBoolProperty>(OutProperty);
 				OutValue = BoolProperty->GetPropertyValue(InPropertyPath.GetCachedAddress());
+				return true;
 			}
-			return true;
+			return false;
 		}
 	};
 
@@ -1179,13 +1193,18 @@ namespace PropertyPathHelpersInternal
 			FArrayProperty* ArrayProp = CastField<FArrayProperty>(Property);
 			if ( ArrayProp && LastSegment.GetArrayIndex() != INDEX_NONE )
 			{
-				ArrayProp->Inner->CopySingleValue(InPropertyPath.GetCachedAddress(), &InValue);
+				if (IsConcreteTypeCompatibleWithReflectedType<T>(ArrayProp->Inner))
+				{
+					ArrayProp->Inner->CopySingleValue(InPropertyPath.GetCachedAddress(), &InValue);
+					return true;
+				}
 			}
-			else
+			else if (IsConcreteTypeCompatibleWithReflectedType<T>(Property))
 			{
 				Property->CopySingleValue(InPropertyPath.GetCachedAddress(), &InValue);
+				return true;
 			}
-			return true;
+			return false;
 		}
 	};
 
@@ -1197,8 +1216,12 @@ namespace PropertyPathHelpersInternal
 		{
 			const FPropertyPathSegment& LastSegment = InPropertyPath.GetLastSegment();
 			FProperty* Property = CastFieldChecked<FProperty>(LastSegment.GetField().ToField());
-			Property->CopyCompleteValue(InPropertyPath.GetCachedAddress(), &InValue);
-			return true;
+			if (IsConcreteTypeCompatibleWithReflectedType<T>(Property))
+			{
+				Property->CopyCompleteValue(InPropertyPath.GetCachedAddress(), &InValue);
+				return true;
+			}
+			return false;
 		}
 	};
 
@@ -1213,15 +1236,20 @@ namespace PropertyPathHelpersInternal
 			FArrayProperty* ArrayProp = CastField<FArrayProperty>(Property);
 			if ( ArrayProp && LastSegment.GetArrayIndex() != INDEX_NONE )
 			{
-				ArrayProp->Inner->CopySingleValue(InPropertyPath.GetCachedAddress(), &InValue);
+				if (IsConcreteTypeCompatibleWithReflectedType<bool>(ArrayProp->Inner))
+				{
+					ArrayProp->Inner->CopySingleValue(InPropertyPath.GetCachedAddress(), &InValue);
+					return true;
+				}
 			}
-			else
+			else if (IsConcreteTypeCompatibleWithReflectedType<bool>(Property))
 			{
 				FBoolProperty* BoolProperty = CastFieldChecked<FBoolProperty>(Property);
 				BoolProperty->SetPropertyValue(InPropertyPath.GetCachedAddress(), InValue);
+				return true;
 			}
 			
-			return true;
+			return false;
 		}
 	};
 
