Make PropertyPathHelpers call setters / getters.
Add support for generic structs to be passed to PropertyPathHelpers.
Fix PropertyPath Not working with TEnumAsByte types
Add TIsTEnumAsByte type trait.

diff --git a/Engine/Source/Runtime/Core/Public/Containers/EnumAsByte.h b/Engine/Source/Runtime/Core/Public/Containers/EnumAsByte.h
--- a/Engine/Source/Runtime/Core/Public/Containers/EnumAsByte.h
+++ b/Engine/Source/Runtime/Core/Public/Containers/EnumAsByte.h
@@ -118,3 +118,14 @@ FORCEINLINE uint32 GetTypeHash(const TEnumAsByte<T>& Enum)
 
 
 template<class T> struct TIsPODType<TEnumAsByte<T>> { enum { Value = true }; };
+
+template <typename T>
+struct TIsTEnumAsByte
+{
+	static constexpr bool Value = false;
+};
+
+template <typename T> struct TIsTEnumAsByte<               TEnumAsByte<T>> { static constexpr bool Value = true; };
+template <typename T> struct TIsTEnumAsByte<const          TEnumAsByte<T>> { static constexpr bool Value = true; };
+template <typename T> struct TIsTEnumAsByte<      volatile TEnumAsByte<T>> { static constexpr bool Value = true; };
+template <typename T> struct TIsTEnumAsByte<const volatile TEnumAsByte<T>> { static constexpr bool Value = true; };
\ No newline at end of file
diff --git a/Engine/Source/Runtime/PropertyPath/Private/PropertyPathHelpers.cpp b/Engine/Source/Runtime/PropertyPath/Private/PropertyPathHelpers.cpp
--- a/Engine/Source/Runtime/PropertyPath/Private/PropertyPathHelpers.cpp
+++ b/Engine/Source/Runtime/PropertyPath/Private/PropertyPathHelpers.cpp
@@ -17,9 +17,7 @@ namespace PropertyPathHelpersInternal
 		// Reset cached address usage flag at the path root. This will be reset later in the recursion if conditions are not met in the path.
 		if(SegmentIndex == 0)
 		{
-#if DO_CHECK
 			InPropertyPath.SetCachedContainer(InContainer);
-#endif
 			InPropertyPath.SetCanSafelyUsedCachedAddress(true);
 		}
 
@@ -47,6 +45,7 @@ namespace PropertyPathHelpersInternal
 						// can continue, if the object is null we safely stop processing the chain of properties.
 						if ( UObject* CurrentObject = ObjectProperty->GetPropertyValue_InContainer(InContainer, ArrayIndex) )
 						{
+							InPropertyPath.SetCachedLastContainer(CurrentObject, SegmentIndex);
 							return IteratePropertyPathRecursive(CurrentObject->GetClass(), CurrentObject, SegmentIndex + 1, InPropertyPath, InResolver);
 						}
 					}
@@ -62,6 +61,7 @@ namespace PropertyPathHelpersInternal
 						// can continue, if the object is null we safely stop processing the chain of properties.
 						if ( UObject* CurrentObject = WeakObject.Get() )
 						{
+							InPropertyPath.SetCachedLastContainer(CurrentObject, SegmentIndex);
 							return IteratePropertyPathRecursive(CurrentObject->GetClass(), CurrentObject, SegmentIndex + 1, InPropertyPath, InResolver);
 						}
 					}
@@ -77,6 +77,7 @@ namespace PropertyPathHelpersInternal
 						// can continue, if the object is null we safely stop processing the chain of properties.
 						if ( UObject* CurrentObject = SoftObject.Get() )
 						{
+							InPropertyPath.SetCachedLastContainer(CurrentObject, SegmentIndex);
 							return IteratePropertyPathRecursive(CurrentObject->GetClass(), CurrentObject, SegmentIndex + 1, InPropertyPath, InResolver);
 						}
 					}
@@ -140,6 +141,7 @@ namespace PropertyPathHelpersInternal
 						if (CurrentObject)
 						{
 							InPropertyPath.SetCanSafelyUsedCachedAddress(false);
+							InPropertyPath.SetCachedLastContainer(CurrentObject, SegmentIndex);
 
 							return IteratePropertyPathRecursive(CurrentObject->GetClass(), CurrentObject, SegmentIndex + 1, InPropertyPath, InResolver);
 						}
@@ -517,6 +519,119 @@ namespace PropertyPathHelpersInternal
 		}
 		return nullptr;
 	}
+
+	void CallParentSetters(const FCachedPropertyPath& InPropertyPath)
+	{
+		int32 LastContainerInPathIndex = InPropertyPath.GetCachedLastContainerInPathIndex();
+		int32 IndexAfterCachedLastContainer = InPropertyPath.GetCachedLastContainerInPathIndex() + 1;
+
+		void* ContainerPtr = LastContainerInPathIndex == INDEX_NONE 
+			? InPropertyPath.GetCachedContainer()
+			: InPropertyPath.GetCachedLastContainerInPath();
+
+		// Call the topmost setter on the last UObject in path
+		const int32 NumSegments = InPropertyPath.GetNumSegments();
+		for (int32 Index = IndexAfterCachedLastContainer; Index < NumSegments; Index++)
+		{
+			const FPropertyPathSegment& ParentSegment = InPropertyPath.GetSegment(Index);
+			int32 ParentArrayIndex = ParentSegment.GetArrayIndex();
+			FProperty* ParentProperty = CastFieldChecked<FProperty>(ParentSegment.GetField().ToField());
+			ParentArrayIndex = ParentArrayIndex == INDEX_NONE ? 0 : ParentArrayIndex;
+
+			if (ParentProperty->HasSetter() && ParentArrayIndex < ParentProperty->ArrayDim)
+			{
+				// We want to call the setter with the current value, so just get the pointer to current value via container
+				if (void* ParentAddress = ParentProperty->ContainerPtrToValuePtr<void>(ContainerPtr, ParentArrayIndex))
+				{
+					if (ParentProperty->HasGetter())
+					{
+						// Call getter if it has one, getter SHOULD be pure and thus won't cause behavioral changes.
+						// But this is a read on ParentAddress so we call getter for now
+
+						int32 Size = ParentSegment.GetStruct()->GetPropertiesSize();
+						int32 Alignment = ParentSegment.GetStruct()->GetMinAlignment();
+						uint8* Temp = (uint8*)FMemory_Alloca_Aligned(Size, Alignment);
+						FMemory::Memzero(Temp, Size);
+
+						if (!ParentProperty->HasAnyPropertyFlags(CPF_ZeroConstructor))
+						{
+							ParentProperty->InitializeValue_InContainer(Temp);
+						}
+						ParentProperty->CallGetter(ContainerPtr, Temp);
+						ParentProperty->CallSetter(ContainerPtr, Temp);
+						ParentProperty->DestroyValue_InContainer(Temp);
+					}
+					else
+					{
+						ParentProperty->CallSetter(ContainerPtr, ParentAddress);
+					}
+					return;
+				}
+			}
+		}
+	}
+
+	void CallParentGetters(void* OutValue, const FCachedPropertyPath& InPropertyPath, const void* InPropertyAddress)
+	{
+		int32 LastContainerInPathIndex = InPropertyPath.GetCachedLastContainerInPathIndex();
+		int32 IndexAfterCachedLastContainer = InPropertyPath.GetCachedLastContainerInPathIndex() + 1;
+
+		void* ContainerPtr = LastContainerInPathIndex == INDEX_NONE 
+			? InPropertyPath.GetCachedContainer()
+			: InPropertyPath.GetCachedLastContainerInPath();
+
+		// Helper to get value regardless of property address / if parent getter is called
+		auto GetValueFromProperty = [](void* OutValue, const FCachedPropertyPath& InPropertyPath, const void* PropertyAddress)
+		{
+			FProperty* Property = CastFieldChecked<FProperty>(InPropertyPath.GetLastSegment().GetField().ToField());
+			if (FBoolProperty* BoolProperty = CastField<FBoolProperty>(Property))
+			{
+				*static_cast<bool*>(OutValue) = BoolProperty->GetPropertyValue(PropertyAddress);
+			}
+			else if (ensure(Property))
+			{
+				Property->CopySingleValue(OutValue, PropertyAddress);
+			}
+		};
+
+		// Call the topmost getter on the last UObject in path
+		const int32 NumSegments = InPropertyPath.GetNumSegments();
+		for (int32 Index = IndexAfterCachedLastContainer; Index < NumSegments; Index++)
+		{
+			const FPropertyPathSegment& ParentSegment = InPropertyPath.GetSegment(Index);
+			int32 ParentArrayIndex = ParentSegment.GetArrayIndex();
+			FProperty* ParentProperty = CastFieldChecked<FProperty>(ParentSegment.GetField().ToField());
+			ParentArrayIndex = ParentArrayIndex == INDEX_NONE ? 0 : ParentArrayIndex;
+
+			if (ParentProperty->HasGetter() && ParentArrayIndex < ParentProperty->ArrayDim)
+			{
+				// We want to call the Getter with the current value, so just get the pointer to current value via container
+				void* ParentAddress = ParentProperty->ContainerPtrToValuePtr<void>(ContainerPtr, ParentArrayIndex);
+				if (ensure(ParentAddress))
+				{
+					int32 Size = ParentSegment.GetStruct()->GetPropertiesSize();
+					int32 Alignment = ParentSegment.GetStruct()->GetMinAlignment();
+					uint8* Temp = (uint8*)FMemory_Alloca_Aligned(Size, Alignment);
+					FMemory::Memzero(Temp, Size);
+
+					if (!ParentProperty->HasAnyPropertyFlags(CPF_ZeroConstructor))
+					{
+						ParentProperty->InitializeValue_InContainer(Temp);
+					}
+					ParentProperty->CallGetter(ContainerPtr, Temp);
+
+					// We resolved the property address earlier & it's containing UObject, use this for relative-offset for Temp
+					uint8* TempPropertyAddress = Temp + ((uint8*)InPropertyAddress - (uint8*)ContainerPtr);
+					GetValueFromProperty(OutValue, InPropertyPath, TempPropertyAddress);
+					ParentProperty->DestroyValue_InContainer(Temp);
+				}
+
+				return;
+			}
+		}
+
+		GetValueFromProperty(OutValue, InPropertyPath, InPropertyAddress);
+	}
 }
 
 FPropertyPathSegment::FPropertyPathSegment()
@@ -601,9 +716,9 @@ UStruct* FPropertyPathSegment::GetStruct() const
 FCachedPropertyPath::FCachedPropertyPath()
 	: CachedAddress(nullptr)
 	, CachedFunction(nullptr)
-#if DO_CHECK
 	, CachedContainer(nullptr)
-#endif
+	, CachedLastContainerInPath(nullptr)
+	, CachedLastContainerInPathIndex(INDEX_NONE)
 	, bCanSafelyUsedCachedAddress(false)
 {
 }
@@ -611,9 +726,9 @@ FCachedPropertyPath::FCachedPropertyPath()
 FCachedPropertyPath::FCachedPropertyPath(const FString& Path)
 	: CachedAddress(nullptr)
 	, CachedFunction(nullptr)
-#if DO_CHECK
 	, CachedContainer(nullptr)
-#endif
+	, CachedLastContainerInPath(nullptr)
+	, CachedLastContainerInPathIndex(INDEX_NONE)
 	, bCanSafelyUsedCachedAddress(false)
 {
 	MakeFromString(Path);
@@ -622,9 +737,9 @@ FCachedPropertyPath::FCachedPropertyPath(const FString& Path)
 FCachedPropertyPath::FCachedPropertyPath(const TArray<FString>& PathSegments)
 	: CachedAddress(nullptr)
 	, CachedFunction(nullptr)
-#if DO_CHECK
 	, CachedContainer(nullptr)
-#endif
+	, CachedLastContainerInPath(nullptr)
+	, CachedLastContainerInPathIndex(INDEX_NONE)
 	, bCanSafelyUsedCachedAddress(false)
 {
 	for (const FString& Segment : PathSegments)
@@ -636,9 +751,9 @@ FCachedPropertyPath::FCachedPropertyPath(const TArray<FString>& PathSegments)
 FCachedPropertyPath::FCachedPropertyPath(const FPropertyPathSegment& Segment)
 	: CachedAddress(nullptr)
 	, CachedFunction(nullptr)
-#if DO_CHECK
 	, CachedContainer(nullptr)
-#endif
+	, CachedLastContainerInPath(nullptr)
+	, CachedLastContainerInPathIndex(INDEX_NONE)
 	, bCanSafelyUsedCachedAddress(false)
 {
 	Segments.Add(Segment);
@@ -721,6 +836,23 @@ void FCachedPropertyPath::SetCanSafelyUsedCachedAddress(bool bInCanSafelyUsedCac
 	bCanSafelyUsedCachedAddress = bInCanSafelyUsedCachedAddress;
 }
 
+void FCachedPropertyPath::SetCachedLastContainer(void* InContainer, int32 InIndex) const
+{
+	SetCanSafelyUsedCachedAddress(false);
+	CachedLastContainerInPath = InContainer;
+	CachedLastContainerInPathIndex = InIndex;
+}
+
+void* FCachedPropertyPath::GetCachedLastContainerInPath() const
+{
+	return CachedLastContainerInPath;
+}
+
+int32 FCachedPropertyPath::GetCachedLastContainerInPathIndex() const
+{
+	return CachedLastContainerInPathIndex;
+}
+
 bool FCachedPropertyPath::IsResolved() const
 {
 	return (CachedFunction != nullptr || CachedAddress != nullptr);
@@ -728,16 +860,14 @@ bool FCachedPropertyPath::IsResolved() const
 
 bool FCachedPropertyPath::IsFullyResolved() const
 {
-#if DO_CHECK
 	bool bCachedContainer = CachedContainer != nullptr;
-#else
-	bool bCachedContainer = true;
-#endif
 	return bCanSafelyUsedCachedAddress && bCachedContainer && IsResolved();
 }
 
 void* FCachedPropertyPath::GetCachedAddress() const
 {
+	// @TODO: DarenC - Should we? Maybe add a GetCachedAddressUnsafe method.
+	// check(bCanSafelyUsedCachedAddress); 
 	return CachedAddress;
 }
 
@@ -817,7 +947,6 @@ bool FCachedPropertyPath::Equals(const FString& Other) const
 	return ToString() == Other;
 }
 
-#if DO_CHECK || USING_CODE_ANALYSIS
 void* FCachedPropertyPath::GetCachedContainer() const
 {
 	return CachedContainer;
@@ -827,7 +956,6 @@ void FCachedPropertyPath::SetCachedContainer(void* InContainer) const
 {
 	CachedContainer = InContainer;
 }
-#endif
 
 void FCachedPropertyPath::RemoveFromEnd(int32 InNumSegments)
 {
@@ -843,9 +971,9 @@ void FCachedPropertyPath::RemoveFromEnd(int32 InNumSegments)
 		}
 		CachedAddress = nullptr;
 		CachedFunction = nullptr;
-#if DO_CHECK
 		CachedContainer = nullptr;
-#endif// DO_CHECK
+		CachedLastContainerInPath = nullptr;
+		CachedLastContainerInPathIndex = INDEX_NONE;
 		bCanSafelyUsedCachedAddress = false;
 	}
 }
@@ -864,9 +992,9 @@ void FCachedPropertyPath::RemoveFromStart(int32 InNumSegments)
 		}
 		CachedAddress = nullptr;
 		CachedFunction = nullptr;
-#if DO_CHECK
 		CachedContainer = nullptr;
-#endif // DO_CHECK
+		CachedLastContainerInPath = nullptr;
+		CachedLastContainerInPathIndex = INDEX_NONE;
 		bCanSafelyUsedCachedAddress = false;
 	}
 }
@@ -1019,6 +1147,18 @@ namespace PropertyPathHelpers
 		return ResolvePropertyPath(InContainer, InStruct, InPropertyPath, Resolver);
 	}
 
+	bool SetPropertyValue(UObject* InContainer, const FCachedPropertyPath& InPropertyPath, const UScriptStruct* InScriptStruct, const uint8* InValue)
+	{
+		PropertyPathHelpersInternal::FPropertyStructView StructView(InScriptStruct, InValue);
+		return SetPropertyValue(InContainer, InPropertyPath, StructView);
+	}
+
+	bool SetPropertyValue(UObject* InContainer, const FString& InPropertyPath, const UScriptStruct* InScriptStruct, const uint8* InValue)
+	{
+		PropertyPathHelpersInternal::FPropertyStructView StructView(InScriptStruct, InValue);
+		return SetPropertyValue(InContainer, InPropertyPath, StructView);
+	}
+
 	bool CopyPropertyValue(UObject* InContainer, const FCachedPropertyPath& InDestPropertyPath, const FCachedPropertyPath& InSrcPropertyPath)
 	{
 		if(InDestPropertyPath.IsFullyResolved() && InSrcPropertyPath.IsFullyResolved())
@@ -1046,10 +1186,8 @@ namespace PropertyPathHelpers
 
 	bool CopyPropertyValueFast(UObject* InContainer, const FCachedPropertyPath& InDestPropertyPath, const FCachedPropertyPath& InSrcPropertyPath)
 	{
-#if DO_CHECK
 		check(InContainer == InDestPropertyPath.GetCachedContainer());
 		check(InContainer == InSrcPropertyPath.GetCachedContainer());
-#endif // DO_CHECK
 		checkSlow(InDestPropertyPath.IsResolved());
 		checkSlow(InSrcPropertyPath.IsResolved());
 		checkSlow(PropertyPathHelpersInternal::CanCopyProperties(InDestPropertyPath, InSrcPropertyPath));
diff --git a/Engine/Source/Runtime/PropertyPath/Private/Tests/PropertyPathHelpersTest.cpp b/Engine/Source/Runtime/PropertyPath/Private/Tests/PropertyPathHelpersTest.cpp
new file mode 100644
--- /dev/null
+++ b/Engine/Source/Runtime/PropertyPath/Private/Tests/PropertyPathHelpersTest.cpp
@@ -0,0 +1,987 @@
+// Copyright Epic Games, Inc. All Rights Reserved.
+
+#include "PropertyPathHelpersTest.h"
+
+#if WITH_DEV_AUTOMATION_TESTS
+
+DEFINE_LOG_CATEGORY_STATIC(LogPropertyPathHelpersTest, Log, All);
+
+/**
+ * Coverage:
+ * 
+ * SetPrimitives (Bool, Integer, Float)
+ * SetInvalidInputIsNoOp
+ * SetString
+ * SetEnum
+ * SetStruct
+ * SetStructData
+ * SetStructMember
+ * SetInnerStructMember
+ * SetPrimitivesInnerObject
+ * SetStringInnerObject
+ * SetStructInnerObject
+ * SetStructMemberInnerObject
+ * SetInnerStructMemberInnerObject
+ * SetInnerStructMemberEnumInnerObject
+ * 
+ * GetPrimitives
+ * GetString
+ * GetStruct
+ * GetEnum
+ * GetStructMember
+ * GetInnerStructMember
+ * GetPrimitivesInnerObject
+ * GetStringInnerObject
+ * GetStructInnerObject
+ * GetStructMemberInnerObject
+ * GetInnerStructMemberInnerObject
+ * GetInnerStructMemberEnumInnerObject
+ * 
+ * Note: 
+ * When a struct setter is called due to a struct member set we expect getter to be called. 
+ * This is because to call the struct setter we had to read the current struct value.
+ */
+
+// Begin Tests
+
+IMPLEMENT_SIMPLE_AUTOMATION_TEST(FSetPrimitivesTest, "System.PropertyPath.SetPrimitivesTest", EAutomationTestFlags::ApplicationContextMask | EAutomationTestFlags::EngineFilter)
+bool FSetPrimitivesTest::RunTest(const FString& Parameters)
+{
+	bool GPassing = true;
+
+	{
+		FPropertyPathTestBed Test = {};
+		FPropertyPathTestBed CppTest = {};
+		CppTest.Object->Bool = true;
+
+		PropertyPathHelpers::SetPropertyValue(Test.Object, FString("Bool"), true);
+		GPassing &= TestEqual(TEXT("Primitive Bool Set"), Test.Object->Bool, true);
+		GPassing &= TestEqual(TEXT("Primitive Bool Set Cpp Equivalent"), *Test.Object, *CppTest.Object);
+		GPassing &= TestEqual(TEXT("Primitive Bool Setter Called"), Test.Object->IsSetterCalled(), false);
+		GPassing &= TestEqual(TEXT("Primitive Bool Getter Called"), Test.Object->IsGetterCalled(), false);
+		Test.Object->ResetGetterSetterFlags();
+	}
+
+	{
+		FPropertyPathTestBed Test = {};
+		FPropertyPathTestBed CppTest = {};
+		CppTest.Object->Integer = 1;
+
+		PropertyPathHelpers::SetPropertyValue(Test.Object, FString("Integer"), 1);
+		GPassing &= TestEqual(TEXT("Primitive Integer Set"), Test.Object->Integer, 1);
+		GPassing &= TestEqual(TEXT("Primitive Integer Set Cpp Equivalent"), *Test.Object, *CppTest.Object);
+		GPassing &= TestEqual(TEXT("Primitive Integer Setter Called"), Test.Object->IsSetterCalled(), false);
+		GPassing &= TestEqual(TEXT("Primitive Integer Getter Called"), Test.Object->IsGetterCalled(), false);
+		Test.Object->ResetGetterSetterFlags();
+	}
+
+	{
+		FPropertyPathTestBed Test = {};
+		FPropertyPathTestBed CppTest = {};
+		CppTest.Object->SetFloat(1.5f);
+
+		PropertyPathHelpers::SetPropertyValue(Test.Object, FString("Float"), 1.5f);
+		GPassing &= TestEqual(TEXT("Primitive Float Set"), Test.Object->Float, 1.5f);
+		GPassing &= TestEqual(TEXT("Primitive Float Set Cpp Equivalent"), *Test.Object, *CppTest.Object);
+		GPassing &= TestEqual(TEXT("Primitive Float Setter Called"), Test.Object->IsSetterCalled(), true);
+		GPassing &= TestEqual(TEXT("Primitive Float Getter Called"), Test.Object->IsGetterCalled(), false);
+		Test.Object->ResetGetterSetterFlags();
+	}
+
+	return GPassing;
+}
+
+IMPLEMENT_SIMPLE_AUTOMATION_TEST(FSetInvalidInputIsNoOp, "System.PropertyPath.SetInvalidInputIsNoOp", EAutomationTestFlags::ApplicationContextMask | EAutomationTestFlags::EngineFilter)
+bool FSetInvalidInputIsNoOp::RunTest(const FString& Parameters)
+{
+	bool GPassing = true;
+
+	// Note: Skipping CppTests here since these are No Ops.
+
+	{
+		FPropertyPathTestBed Test = {};
+
+		PropertyPathHelpers::SetPropertyValue(Test.Object, FString("Bool"), Test.ModifiedStruct);
+		GPassing &= TestEqual(TEXT("Bool Invalid Input Set No Op"), Test.Object->Bool, false);
+		GPassing &= TestEqual(TEXT("Bool Setter Called"), Test.Object->IsSetterCalled(), false);
+		GPassing &= TestEqual(TEXT("Bool Getter Called"), Test.Object->IsGetterCalled(), false);
+		Test.Object->ResetGetterSetterFlags();
+	}
+
+	{
+		FPropertyPathTestBed Test = {};
+
+		PropertyPathHelpers::SetPropertyValue(Test.Object, FString("Integer"), Test.ModifiedStruct);
+		GPassing &= TestEqual(TEXT("Integer Invalid Input Set No Op"), Test.Object->Integer, 0);
+		GPassing &= TestEqual(TEXT("Integer Setter Called"), Test.Object->IsSetterCalled(), false);
+		GPassing &= TestEqual(TEXT("Integer Getter Called"), Test.Object->IsGetterCalled(), false);
+		Test.Object->ResetGetterSetterFlags();
+	}
+
+	{
+		FPropertyPathTestBed Test = {};
+
+		PropertyPathHelpers::SetPropertyValue(Test.Object, FString("Float"), Test.ModifiedStruct);
+		GPassing &= TestEqual(TEXT("Float Invalid Input Set No Op"), Test.Object->Float, 0.0f);
+		GPassing &= TestEqual(TEXT("Float Setter Called"), Test.Object->IsSetterCalled(), false);
+		GPassing &= TestEqual(TEXT("Float Getter Called"), Test.Object->IsGetterCalled(), false);
+		Test.Object->ResetGetterSetterFlags();
+	}
+
+	{
+		FPropertyPathTestBed Test = {};
+
+		PropertyPathHelpers::SetPropertyValue(Test.Object, FString("Struct"), 1);
+		GPassing &= TestEqual(TEXT("Struct Invalid Input Set No Op"), Test.Object->Struct, Test.DefaultStruct);
+		GPassing &= TestEqual(TEXT("Struct Setter Called"), Test.Object->IsSetterCalled(), false);
+		GPassing &= TestEqual(TEXT("Struct Getter Called"), Test.Object->IsGetterCalled(), false);
+		Test.Object->ResetGetterSetterFlags();
+	}
+
+	return GPassing;
+}
+
+IMPLEMENT_SIMPLE_AUTOMATION_TEST(FSetStringTest, "System.PropertyPath.SetStringTest", EAutomationTestFlags::ApplicationContextMask | EAutomationTestFlags::EngineFilter)
+bool FSetStringTest::RunTest(const FString& Parameters)
+{
+	bool GPassing = true;
+
+	{
+		FPropertyPathTestBed Test = {};
+		FPropertyPathTestBed CppTest = {};
+		CppTest.Object->String = FString("NewValue");
+
+		PropertyPathHelpers::SetPropertyValue(Test.Object, FString("String"), FString("NewValue"));
+		GPassing &= TestEqual(TEXT("String Set"), Test.Object->String, FString("NewValue"));
+		GPassing &= TestEqual(TEXT("String Set Cpp Equivalent"), *Test.Object, *CppTest.Object);
+		GPassing &= TestEqual(TEXT("String Setter Called"), Test.Object->IsSetterCalled(), false);
+		GPassing &= TestEqual(TEXT("String Getter Called"), Test.Object->IsGetterCalled(), false);
+		Test.Object->ResetGetterSetterFlags();
+	}
+
+	return GPassing;
+}
+
+IMPLEMENT_SIMPLE_AUTOMATION_TEST(FSetEnumTest, "System.PropertyPath.SetEnumTest", EAutomationTestFlags::ApplicationContextMask | EAutomationTestFlags::EngineFilter)
+bool FSetEnumTest::RunTest(const FString& Parameters)
+{
+	bool GPassing = true;
+
+	{
+		FPropertyPathTestBed Test = {};
+		FPropertyPathTestBed CppTest = {};
+		CppTest.Object->EnumTwo = Two;
+
+		PropertyPathHelpers::SetPropertyValue(Test.Object, FString("EnumTwo"), Two);
+		GPassing &= TestEqual(TEXT("Enum Set"), Test.Object->EnumTwo, Two);
+		GPassing &= TestEqual(TEXT("Enum Set Cpp Equivalent"), *Test.Object, *CppTest.Object);
+		GPassing &= TestEqual(TEXT("Enum Setter Called"), Test.Object->IsSetterCalled(), false);
+		GPassing &= TestEqual(TEXT("Enum Getter Called"), Test.Object->IsGetterCalled(), false);
+		Test.Object->ResetGetterSetterFlags();
+	}
+
+	{
+		// Test for
+		FPropertyPathTestBed Test = {};
+		FPropertyPathTestBed CppTest = {};
+		CppTest.Object->EnumOne = Four;
+		CppTest.Object->EnumTwo = Three;
+		CppTest.Object->EnumThree = Two;
+		CppTest.Object->EnumFour = One;
+
+		PropertyPathHelpers::SetPropertyValue(Test.Object, FString("EnumOne"), Four);
+		CppTest.Object->EnumOne = Four;
+		GPassing &= TestEqual(TEXT("Enum Set Cpp Equivalent"), *Test.Object, *CppTest.Object);
+
+		PropertyPathHelpers::SetPropertyValue(Test.Object, FString("EnumTwo"), Three);
+		CppTest.Object->EnumTwo = Three;
+		GPassing &= TestEqual(TEXT("Enum Set Cpp Equivalent"), *Test.Object, *CppTest.Object);
+
+		PropertyPathHelpers::SetPropertyValue(Test.Object, FString("EnumThree"), Two);
+		CppTest.Object->EnumThree = Two;
+		GPassing &= TestEqual(TEXT("Enum Set Cpp Equivalent"), *Test.Object, *CppTest.Object);
+
+		PropertyPathHelpers::SetPropertyValue(Test.Object, FString("EnumFour"), One);
+		CppTest.Object->EnumFour = One;
+		GPassing &= TestEqual(TEXT("Enum Set Cpp Equivalent"), *Test.Object, *CppTest.Object);
+
+		GPassing &= TestEqual(TEXT("Enum Close No Overwrite Set"), Test.Object->EnumOne, Four);
+		GPassing &= TestEqual(TEXT("Enum Close No Overwrite Set"), Test.Object->EnumTwo, Three);
+		GPassing &= TestEqual(TEXT("Enum Close No Overwrite Set"), Test.Object->EnumThree, Two);
+		GPassing &= TestEqual(TEXT("Enum Close No Overwrite Set"), Test.Object->EnumFour, One);
+
+		PropertyPathHelpers::SetPropertyValue(Test.Object, FString("EnumFour"), Four);
+		CppTest.Object->EnumFour = Four;
+		GPassing &= TestEqual(TEXT("Enum Set Cpp Equivalent"), *Test.Object, *CppTest.Object);
+
+		PropertyPathHelpers::SetPropertyValue(Test.Object, FString("EnumThree"), Three);
+		CppTest.Object->EnumThree = Three;
+		GPassing &= TestEqual(TEXT("Enum Set Cpp Equivalent"), *Test.Object, *CppTest.Object);
+
+		PropertyPathHelpers::SetPropertyValue(Test.Object, FString("EnumTwo"), Two);
+		CppTest.Object->EnumTwo = Two;
+		GPassing &= TestEqual(TEXT("Enum Set Cpp Equivalent"), *Test.Object, *CppTest.Object);
+
+		PropertyPathHelpers::SetPropertyValue(Test.Object, FString("EnumOne"), One);
+		CppTest.Object->EnumOne = One;
+		GPassing &= TestEqual(TEXT("Enum Set Cpp Equivalent"), *Test.Object, *CppTest.Object);
+
+		GPassing &= TestEqual(TEXT("Enum Close No Overwrite Set"), Test.Object->EnumOne, One);
+		GPassing &= TestEqual(TEXT("Enum Close No Overwrite Set"), Test.Object->EnumTwo, Two);
+		GPassing &= TestEqual(TEXT("Enum Close No Overwrite Set"), Test.Object->EnumThree, Three);
+		GPassing &= TestEqual(TEXT("Enum Close No Overwrite Set"), Test.Object->EnumFour, Four);
+
+		GPassing &= TestEqual(TEXT("Enum Setter Called"), Test.Object->IsSetterCalled(), false);
+		GPassing &= TestEqual(TEXT("Enum Getter Called"), Test.Object->IsGetterCalled(), false);
+		Test.Object->ResetGetterSetterFlags();
+	}
+
+	return GPassing;
+}
+
+IMPLEMENT_SIMPLE_AUTOMATION_TEST(FSetStructTest, "System.PropertyPath.SetStructTest", EAutomationTestFlags::ApplicationContextMask | EAutomationTestFlags::EngineFilter)
+bool FSetStructTest::RunTest(const FString& Parameters)
+{
+	bool GPassing = true;
+
+	{
+		FPropertyPathTestBed Test = {};
+		FPropertyPathTestBed CppTest = {};
+		CppTest.Object->SetStruct(Test.ModifiedStruct);
+
+		PropertyPathHelpers::SetPropertyValue(Test.Object, FString("Struct"), Test.ModifiedStruct);
+		GPassing &= TestEqual(TEXT("Struct Set"), Test.Object->Struct, Test.ModifiedStruct);
+		GPassing &= TestEqual(TEXT("Struct Set Cpp Equivalent"), *Test.Object, *CppTest.Object);
+		GPassing &= TestEqual(TEXT("Struct Setter Called"), Test.Object->IsSetterCalled(), true);
+		GPassing &= TestEqual(TEXT("Struct Getter Called"), Test.Object->IsGetterCalled(), false);
+		Test.Object->ResetGetterSetterFlags();
+	}
+
+	{
+		FPropertyPathTestBed Test = {};
+		FPropertyPathTestBed CppTest = {};
+		CppTest.Object->SetStructRef(Test.ModifiedStruct);
+
+		PropertyPathHelpers::SetPropertyValue(Test.Object, FString("StructRef"), Test.ModifiedStruct);
+		GPassing &= TestEqual(TEXT("Struct Ref Set"), Test.Object->StructRef, Test.ModifiedStruct);
+		GPassing &= TestEqual(TEXT("Struct Ref Set Cpp Equivalent"), *Test.Object, *CppTest.Object);
+		GPassing &= TestEqual(TEXT("Struct Setter Called"), Test.Object->IsSetterCalled(), true);
+		GPassing &= TestEqual(TEXT("Struct Getter Called"), Test.Object->IsGetterCalled(), false);
+		Test.Object->ResetGetterSetterFlags();
+	}
+
+	{
+		FPropertyPathTestBed Test = {};
+		FPropertyPathTestBed CppTest = {};
+		CppTest.Object->SetStructConstRef(Test.ModifiedStruct);
+
+		PropertyPathHelpers::SetPropertyValue(Test.Object, FString("StructConstRef"), Test.ModifiedStruct);
+		GPassing &= TestEqual(TEXT("Struct Const Ref Set"), Test.Object->StructConstRef, Test.ModifiedStruct);
+		GPassing &= TestEqual(TEXT("Struct Const Ref Set Cpp Equivalent"), *Test.Object, *CppTest.Object);
+		GPassing &= TestEqual(TEXT("Struct Setter Called"), Test.Object->IsSetterCalled(), true);
+		GPassing &= TestEqual(TEXT("Struct Getter Called"), Test.Object->IsGetterCalled(), false);
+		Test.Object->ResetGetterSetterFlags();
+	}
+
+	return GPassing;
+}
+
+IMPLEMENT_SIMPLE_AUTOMATION_TEST(FSetStructDataTest, "System.PropertyPath.SetStructDataTest", EAutomationTestFlags::ApplicationContextMask | EAutomationTestFlags::EngineFilter)
+bool FSetStructDataTest::RunTest(const FString& Parameters)
+{
+	bool GPassing = true;
+
+	{
+		FPropertyPathTestBed Test = {};
+		FPropertyPathTestBed CppTest = {};
+		CppTest.Object->SetStruct(Test.ModifiedStruct);
+		
+		UScriptStruct* ScriptStruct = TBaseStructure<FPropertyPathTestStruct>::Get();
+		PropertyPathHelpers::SetPropertyValue(Test.Object, FString("Struct"), ScriptStruct, (const uint8*)&Test.ModifiedStruct);
+		GPassing &= TestEqual(TEXT("Struct Data Set"), Test.Object->Struct, Test.ModifiedStruct);
+		GPassing &= TestEqual(TEXT("Struct Data Set Cpp Equivalent"), *Test.Object, *CppTest.Object);
+		GPassing &= TestEqual(TEXT("Struct Setter Called"), Test.Object->IsSetterCalled(), true);
+		GPassing &= TestEqual(TEXT("Struct Getter Called"), Test.Object->IsGetterCalled(), false);
+		Test.Object->ResetGetterSetterFlags();
+	}
+
+	{
+		FPropertyPathTestBed Test = {};
+		FPropertyPathTestBed CppTest = {};
+		CppTest.Object->SetStructRef(Test.ModifiedStruct);
+
+		UScriptStruct* ScriptStruct = TBaseStructure<FPropertyPathTestStruct>::Get();
+		PropertyPathHelpers::SetPropertyValue(Test.Object, FString("StructRef"), ScriptStruct, (const uint8*)&Test.ModifiedStruct);
+		GPassing &= TestEqual(TEXT("Struct Ref Data Set"), Test.Object->StructRef, Test.ModifiedStruct);
+		GPassing &= TestEqual(TEXT("Struct Ref Data Set Cpp Equivalent"), *Test.Object, *CppTest.Object);
+		GPassing &= TestEqual(TEXT("Struct Setter Called"), Test.Object->IsSetterCalled(), true);
+		GPassing &= TestEqual(TEXT("Struct Getter Called"), Test.Object->IsGetterCalled(), false);
+		Test.Object->ResetGetterSetterFlags();
+	}
+
+	{
+		FPropertyPathTestBed Test = {};
+		FPropertyPathTestBed CppTest = {};
+		CppTest.Object->SetStructConstRef(Test.ModifiedStruct);
+
+		UScriptStruct* ScriptStruct = TBaseStructure<FPropertyPathTestStruct>::Get();
+		PropertyPathHelpers::SetPropertyValue(Test.Object, FString("StructConstRef"), ScriptStruct, (const uint8*)&Test.ModifiedStruct);
+		GPassing &= TestEqual(TEXT("Struct Const Ref Data Set"), Test.Object->StructConstRef, Test.ModifiedStruct);
+		GPassing &= TestEqual(TEXT("Struct Const Ref Data Set Cpp Equivalent"), *Test.Object, *CppTest.Object);
+		GPassing &= TestEqual(TEXT("Struct Setter Called"), Test.Object->IsSetterCalled(), true);
+		GPassing &= TestEqual(TEXT("Struct Getter Called"), Test.Object->IsGetterCalled(), false);
+		Test.Object->ResetGetterSetterFlags();
+	}
+
+	return GPassing;
+}
+
+IMPLEMENT_SIMPLE_AUTOMATION_TEST(FSetStructMemberTest, "System.PropertyPath.SetStructMemberTest", EAutomationTestFlags::ApplicationContextMask | EAutomationTestFlags::EngineFilter)
+bool FSetStructMemberTest::RunTest(const FString& Parameters)
+{
+	bool GPassing = true;
+
+	{
+		FPropertyPathTestBed Test = {};
+		FPropertyPathTestBed CppTest = {};
+		CppTest.Object->Struct.Float = 1.5f;
+		CppTest.Object->SetStruct(CppTest.Object->Struct);
+
+		PropertyPathHelpers::SetPropertyValue(Test.Object, FString("Struct.Float"), 1.5f);
+		GPassing &= TestEqual(TEXT("Struct Member Set"), Test.Object->Struct.Float, 1.5f);
+		GPassing &= TestEqual(TEXT("Struct Member Set Cpp Equivalent"), *Test.Object, *CppTest.Object);
+		GPassing &= TestEqual(TEXT("Struct Setter Called"), Test.Object->IsSetterCalled(), true);
+
+		// Emphasis, see note under coverage at top
+		GPassing &= TestEqual(TEXT("Struct Getter Called"), Test.Object->IsGetterCalled(), true);
+
+		Test.Object->ResetGetterSetterFlags();
+	}
+
+	{
+		FPropertyPathTestBed Test = {};
+		FPropertyPathTestBed CppTest = {};
+		CppTest.Object->StructRef.Float = 1.5f;
+		CppTest.Object->SetStructRef(CppTest.Object->StructRef);
+
+		PropertyPathHelpers::SetPropertyValue(Test.Object, FString("StructRef.Float"), 1.5f);
+		GPassing &= TestEqual(TEXT("Struct Ref Member Set"), Test.Object->StructRef.Float, 1.5f);
+		GPassing &= TestEqual(TEXT("Struct Ref Member Set Cpp Equivalent"), *Test.Object, *CppTest.Object);
+		GPassing &= TestEqual(TEXT("Struct Setter Called"), Test.Object->IsSetterCalled(), true);
+
+		// Emphasis, see note under coverage at top
+		GPassing &= TestEqual(TEXT("Struct Getter Called"), Test.Object->IsGetterCalled(), true);
+
+		Test.Object->ResetGetterSetterFlags();
+	}
+
+	{
+		FPropertyPathTestBed Test = {};
+		FPropertyPathTestBed CppTest = {};
+		CppTest.Object->StructConstRef.Float = 1.5f;
+		CppTest.Object->SetStructConstRef(CppTest.Object->StructConstRef);
+
+		PropertyPathHelpers::SetPropertyValue(Test.Object, FString("StructConstRef.Float"), 1.5f);
+		GPassing &= TestEqual(TEXT("Struct Const Ref Member Set"), Test.Object->StructConstRef.Float, 1.5f);
+		GPassing &= TestEqual(TEXT("Struct Const Ref Member Set Cpp Equivalent"), *Test.Object, *CppTest.Object);
+		GPassing &= TestEqual(TEXT("Struct Setter Called"), Test.Object->IsSetterCalled(), true);
+
+		// Emphasis, see note under coverage at top
+		GPassing &= TestEqual(TEXT("Struct Getter Called"), Test.Object->IsGetterCalled(), true);
+
+		Test.Object->ResetGetterSetterFlags();
+	}
+
+	return GPassing;
+}
+
+IMPLEMENT_SIMPLE_AUTOMATION_TEST(FSetInnerStructMemberTest, "System.PropertyPath.SetInnerStructMemberTest", EAutomationTestFlags::ApplicationContextMask | EAutomationTestFlags::EngineFilter)
+bool FSetInnerStructMemberTest::RunTest(const FString& Parameters)
+{
+	bool GPassing = true;
+
+	{
+		FPropertyPathTestBed Test = {};
+		FPropertyPathTestBed CppTest = {};
+		CppTest.Object->Struct.InnerStruct.Float = 1.5f;
+		CppTest.Object->SetStruct(CppTest.Object->Struct);
+
+		PropertyPathHelpers::SetPropertyValue(Test.Object, FString("Struct.InnerStruct.Float"), 1.5f);
+		GPassing &= TestEqual(TEXT("Inner Struct Member Set"), Test.Object->Struct.InnerStruct.Float, 1.5f);
+		GPassing &= TestEqual(TEXT("Inner Struct Member Set Cpp Equivalent"), *Test.Object, *CppTest.Object);
+		GPassing &= TestEqual(TEXT("Struct Setter Called"), Test.Object->IsSetterCalled(), true);
+
+		// Emphasis, see note under coverage at top
+		GPassing &= TestEqual(TEXT("Struct Getter Called"), Test.Object->IsGetterCalled(), true);
+
+		Test.Object->ResetGetterSetterFlags();
+	}
+
+	return GPassing;
+}
+
+IMPLEMENT_SIMPLE_AUTOMATION_TEST(FSetPrimitivesInnerObjectTest, "System.PropertyPath.SetPrimitivesInnerObjectTest", EAutomationTestFlags::ApplicationContextMask | EAutomationTestFlags::EngineFilter)
+bool FSetPrimitivesInnerObjectTest::RunTest(const FString& Parameters)
+{
+	bool GPassing = true;
+
+	{
+		FPropertyPathTestBed Test = {};
+		FPropertyPathTestBed CppTest = {};
+		CppTest.Object->InnerObject->Bool = true;
+
+		PropertyPathHelpers::SetPropertyValue(Test.Object, FString("InnerObject.Bool"), true);
+		GPassing &= TestEqual(TEXT("Primitive Inner Object Bool Set"), Test.Object->InnerObject->Bool, true);
+		GPassing &= TestEqual(TEXT("Primitive Inner Object Bool Set Cpp Equivalent"), *Test.Object, *CppTest.Object);
+		GPassing &= TestEqual(TEXT("Primitive Inner Object Bool Setter Called"), Test.Object->InnerObject->IsSetterCalled(), false);
+		GPassing &= TestEqual(TEXT("Primitive Inner Object Bool Getter Called"), Test.Object->InnerObject->IsGetterCalled(), false);
+		GPassing &= TestEqual(TEXT("Primitive Bool Setter Called"), Test.Object->IsSetterCalled(), false);
+		GPassing &= TestEqual(TEXT("Primitive Bool Getter Called"), Test.Object->IsGetterCalled(), false);
+		Test.Object->InnerObject->ResetGetterSetterFlags();
+		Test.Object->ResetGetterSetterFlags();
+	}
+
+	{
+		FPropertyPathTestBed Test = {};
+		FPropertyPathTestBed CppTest = {};
+		CppTest.Object->InnerObject->Integer = 1;
+
+		PropertyPathHelpers::SetPropertyValue(Test.Object, FString("InnerObject.Integer"), 1);
+		GPassing &= TestEqual(TEXT("Primitive Inner Object Integer Set"), Test.Object->InnerObject->Integer, 1);
+		GPassing &= TestEqual(TEXT("Primitive Inner Object Integer Set Cpp Equivalent"), *Test.Object, *CppTest.Object);
+		GPassing &= TestEqual(TEXT("Primitive Inner Object Integer Setter Called"), Test.Object->InnerObject->IsSetterCalled(), false);
+		GPassing &= TestEqual(TEXT("Primitive Inner Object Integer Getter Called"), Test.Object->InnerObject->IsGetterCalled(), false);
+		GPassing &= TestEqual(TEXT("Primitive Bool Setter Called"), Test.Object->IsSetterCalled(), false);
+		GPassing &= TestEqual(TEXT("Primitive Bool Getter Called"), Test.Object->IsGetterCalled(), false);
+		Test.Object->InnerObject->ResetGetterSetterFlags();
+		Test.Object->ResetGetterSetterFlags();
+	}
+
+	{
+		FPropertyPathTestBed Test = {};
+		FPropertyPathTestBed CppTest = {};
+		CppTest.Object->InnerObject->SetFloat(1.5f);
+
+		PropertyPathHelpers::SetPropertyValue(Test.Object, FString("InnerObject.Float"), 1.5f);
+		GPassing &= TestEqual(TEXT("Primitive Inner Object Float Set"), Test.Object->InnerObject->Float, 1.5f);
+		GPassing &= TestEqual(TEXT("Primitive Inner Object Float Set Cpp Equivalent"), *Test.Object, *CppTest.Object);
+		GPassing &= TestEqual(TEXT("Primitive Inner Object Float Setter Called"), Test.Object->InnerObject->IsSetterCalled(), true);
+		GPassing &= TestEqual(TEXT("Primitive Inner Object Float Getter Called"), Test.Object->InnerObject->IsGetterCalled(), false);
+		GPassing &= TestEqual(TEXT("Primitive Float Setter Called"), Test.Object->IsSetterCalled(), false);
+		GPassing &= TestEqual(TEXT("Primitive Float Getter Called"), Test.Object->IsGetterCalled(), false);
+		Test.Object->InnerObject->ResetGetterSetterFlags();
+		Test.Object->ResetGetterSetterFlags();
+	}
+
+	return GPassing;
+}
+
+IMPLEMENT_SIMPLE_AUTOMATION_TEST(FSetStringInnerObjectTest, "System.PropertyPath.SetStringInnerObjectTest", EAutomationTestFlags::ApplicationContextMask | EAutomationTestFlags::EngineFilter)
+bool FSetStringInnerObjectTest::RunTest(const FString& Parameters)
+{
+	bool GPassing = true;
+
+	{
+		FPropertyPathTestBed Test = {};
+		FPropertyPathTestBed CppTest = {};
+		CppTest.Object->InnerObject->String = FString("NewValue");
+
+		PropertyPathHelpers::SetPropertyValue(Test.Object, FString("InnerObject.String"), FString("NewValue"));
+		GPassing &= TestEqual(TEXT("Inner Object String Set"), Test.Object->InnerObject->String, FString("NewValue"));
+		GPassing &= TestEqual(TEXT("Inner Object String Set Cpp Equivalent"), *Test.Object, *CppTest.Object);
+		GPassing &= TestEqual(TEXT("Inner Object String Setter Called"), Test.Object->InnerObject->IsSetterCalled(), false);
+		GPassing &= TestEqual(TEXT("Inner Object String Getter Called"), Test.Object->InnerObject->IsGetterCalled(), false);
+		GPassing &= TestEqual(TEXT("String Setter Called"), Test.Object->IsSetterCalled(), false);
+		GPassing &= TestEqual(TEXT("String Getter Called"), Test.Object->IsGetterCalled(), false);
+		Test.Object->InnerObject->ResetGetterSetterFlags();
+		Test.Object->ResetGetterSetterFlags();
+	}
+
+	return GPassing;
+}
+
+IMPLEMENT_SIMPLE_AUTOMATION_TEST(FSetStructInnerObjectTest, "System.PropertyPath.SetStructInnerObjectTest", EAutomationTestFlags::ApplicationContextMask | EAutomationTestFlags::EngineFilter)
+bool FSetStructInnerObjectTest::RunTest(const FString& Parameters)
+{
+	bool GPassing = true;
+
+	{
+		FPropertyPathTestBed Test = {};
+		FPropertyPathTestBed CppTest = {};
+		CppTest.Object->InnerObject->SetStruct(Test.ModifiedStruct);
+
+		PropertyPathHelpers::SetPropertyValue(Test.Object, FString("InnerObject.Struct"), Test.ModifiedStruct);
+		GPassing &= TestEqual(TEXT("Inner Object Struct Set"), Test.Object->InnerObject->Struct, Test.ModifiedStruct);
+		GPassing &= TestEqual(TEXT("Inner Object Struct Set Cpp Equivalent"), *Test.Object, *CppTest.Object);
+		GPassing &= TestEqual(TEXT("Inner Object Struct Setter Called"), Test.Object->InnerObject->IsSetterCalled(), true);
+		GPassing &= TestEqual(TEXT("Inner Object Struct Getter Called"), Test.Object->InnerObject->IsGetterCalled(), false);
+		GPassing &= TestEqual(TEXT("Struct Setter Called"), Test.Object->IsSetterCalled(), false);
+		GPassing &= TestEqual(TEXT("Struct Getter Called"), Test.Object->IsGetterCalled(), false);
+		Test.Object->InnerObject->ResetGetterSetterFlags();
+		Test.Object->ResetGetterSetterFlags();
+	}
+	return GPassing;
+}
+
+IMPLEMENT_SIMPLE_AUTOMATION_TEST(FSetStructMemberInnerObjectTest, "System.PropertyPath.SetStructMemberInnerObjectTest", EAutomationTestFlags::ApplicationContextMask | EAutomationTestFlags::EngineFilter)
+bool FSetStructMemberInnerObjectTest::RunTest(const FString& Parameters)
+{
+	bool GPassing = true;
+
+	{
+		FPropertyPathTestBed Test = {};
+		FPropertyPathTestBed CppTest = {};
+		CppTest.Object->InnerObject->Struct.Float = 1.5f;
+		CppTest.Object->InnerObject->SetStruct(CppTest.Object->InnerObject->Struct);
+
+		PropertyPathHelpers::SetPropertyValue(Test.Object, FString("InnerObject.Struct.Float"), 1.5f);
+		GPassing &= TestEqual(TEXT("Inner Object Struct Member Set"), Test.Object->InnerObject->Struct.Float, 1.5f);
+		GPassing &= TestEqual(TEXT("Inner Object Struct Member Set Cpp Equivalent"), *Test.Object, *CppTest.Object);
+		GPassing &= TestEqual(TEXT("Inner Object Struct Setter Called"), Test.Object->InnerObject->IsSetterCalled(), true);
+
+		// Emphasis, see note under coverage at top
+		GPassing &= TestEqual(TEXT("Inner Object Struct Getter Called"), Test.Object->InnerObject->IsGetterCalled(), true);
+
+		GPassing &= TestEqual(TEXT("Struct Setter Called"), Test.Object->IsSetterCalled(), false);
+		GPassing &= TestEqual(TEXT("Struct Getter Called"), Test.Object->IsGetterCalled(), false);
+		Test.Object->ResetGetterSetterFlags();
+	}
+
+	return GPassing;
+}
+
+IMPLEMENT_SIMPLE_AUTOMATION_TEST(FSetInnerStructMemberInnerObjectTest, "System.PropertyPath.SetInnerStructMemberInnerObjectTest", EAutomationTestFlags::ApplicationContextMask | EAutomationTestFlags::EngineFilter)
+bool FSetInnerStructMemberInnerObjectTest::RunTest(const FString& Parameters)
+{
+	bool GPassing = true;
+
+	{
+		FPropertyPathTestBed Test = {};
+		FPropertyPathTestBed CppTest = {};
+		CppTest.Object->InnerObject->Struct.InnerStruct.Float = 1.5f;
+		CppTest.Object->InnerObject->SetStruct(CppTest.Object->InnerObject->Struct);
+
+		PropertyPathHelpers::SetPropertyValue(Test.Object, FString("InnerObject.Struct.InnerStruct.Float"), 1.5f);
+		GPassing &= TestEqual(TEXT("Inner Object Inner Struct Member Set"), Test.Object->InnerObject->Struct.InnerStruct.Float, 1.5f);
+		GPassing &= TestEqual(TEXT("Inner Object Inner Struct Member Set Cpp Equivalent"), *Test.Object, *CppTest.Object);
+		GPassing &= TestEqual(TEXT("Inner Object Inner Struct Setter Called"), Test.Object->InnerObject->IsSetterCalled(), true);
+
+		// Emphasis, see note under coverage at top
+		GPassing &= TestEqual(TEXT("Inner Object Inner Struct Getter Called"), Test.Object->InnerObject->IsGetterCalled(), true);
+
+		GPassing &= TestEqual(TEXT("Inner Struct Setter Called"), Test.Object->IsSetterCalled(), false);
+		GPassing &= TestEqual(TEXT("Inner Struct Getter Called"), Test.Object->IsGetterCalled(), false);
+		Test.Object->ResetGetterSetterFlags();
+	}
+
+	return GPassing;
+}
+
+IMPLEMENT_SIMPLE_AUTOMATION_TEST(FSetInnerStructMemberEnumInnerObjectTest, "System.PropertyPath.SetInnerStructMemberEnumInnerObjectTest", EAutomationTestFlags::ApplicationContextMask | EAutomationTestFlags::EngineFilter)
+bool FSetInnerStructMemberEnumInnerObjectTest::RunTest(const FString& Parameters)
+{
+	bool GPassing = true;
+
+	{
+		FPropertyPathTestBed Test = {};
+		FPropertyPathTestBed CppTest = {};
+		CppTest.Object->InnerObject->Struct.InnerStruct.EnumTwo = Two;
+		CppTest.Object->InnerObject->SetStruct(CppTest.Object->InnerObject->Struct);
+
+		PropertyPathHelpers::SetPropertyValue(Test.Object, FString("InnerObject.Struct.InnerStruct.EnumTwo"), Two);
+		GPassing &= TestEqual(TEXT("Inner Object Inner Struct Member Enum Set"), Test.Object->InnerObject->Struct.InnerStruct.EnumTwo, Two);
+		GPassing &= TestEqual(TEXT("Inner Object Inner Struct Member Enum Set Cpp Equivalent"), *Test.Object, *CppTest.Object);
+		GPassing &= TestEqual(TEXT("Inner Object Inner Struct Setter Called"), Test.Object->InnerObject->IsSetterCalled(), true);
+
+		// Emphasis, see note under coverage at top
+		GPassing &= TestEqual(TEXT("Inner Object Inner Struct Getter Called"), Test.Object->InnerObject->IsGetterCalled(), true);
+
+		GPassing &= TestEqual(TEXT("Inner Struct Setter Called"), Test.Object->IsSetterCalled(), false);
+		GPassing &= TestEqual(TEXT("Inner Struct Getter Called"), Test.Object->IsGetterCalled(), false);
+		Test.Object->ResetGetterSetterFlags();
+	}
+
+	return GPassing;
+}
+
+ IMPLEMENT_SIMPLE_AUTOMATION_TEST(FGetPrimitivesTest, "System.PropertyPath.GetPrimitivesTest", EAutomationTestFlags::ApplicationContextMask | EAutomationTestFlags::EngineFilter)
+ bool FGetPrimitivesTest::RunTest(const FString& Parameters)
+ {
+	 bool GPassing = true;
+
+	 {
+	 	FPropertyPathTestBed Test = {};
+		FPropertyPathTestBed CppTest = {};
+		CppTest.ModifiedStruct.Bool = CppTest.Object->Bool;
+
+	 	PropertyPathHelpers::GetPropertyValue(Test.Object, FString("Bool"), Test.ModifiedStruct.Bool);
+	 	GPassing &= TestEqual(TEXT("Primitive Bool Get"), Test.DefaultStruct.Bool, Test.ModifiedStruct.Bool);
+		GPassing &= TestEqual(TEXT("Primitive Bool Get Cpp Equivalent"), CppTest.ModifiedStruct, Test.ModifiedStruct);
+	 	GPassing &= TestEqual(TEXT("Primitive Bool Setter Called"), Test.Object->IsSetterCalled(), false);
+	 	GPassing &= TestEqual(TEXT("Primitive Bool Getter Called"), Test.Object->IsGetterCalled(), false);
+	 	Test.Object->ResetGetterSetterFlags();
+	}
+
+	{
+		FPropertyPathTestBed Test = {};
+		FPropertyPathTestBed CppTest = {};
+		CppTest.ModifiedStruct.Integer = CppTest.Object->Integer;
+
+	 	PropertyPathHelpers::GetPropertyValue(Test.Object, FString("Integer"), Test.ModifiedStruct.Integer);
+	 	GPassing &= TestEqual(TEXT("Primitive Integer Get"), Test.DefaultStruct.Integer, Test.ModifiedStruct.Integer);
+		GPassing &= TestEqual(TEXT("Primitive Integer Get Cpp Equivalent"), CppTest.ModifiedStruct, Test.ModifiedStruct);
+	 	GPassing &= TestEqual(TEXT("Primitive Integer Setter Called"), Test.Object->IsSetterCalled(), false);
+	 	GPassing &= TestEqual(TEXT("Primitive Integer Getter Called"), Test.Object->IsGetterCalled(), false);
+	 	Test.Object->ResetGetterSetterFlags();
+	}
+
+	{
+		FPropertyPathTestBed Test = {};
+		FPropertyPathTestBed CppTest = {};
+		CppTest.ModifiedStruct.Float = CppTest.Object->GetFloat();
+
+	 	PropertyPathHelpers::GetPropertyValue(Test.Object, FString("Float"), Test.ModifiedStruct.Float);
+	 	GPassing &= TestEqual(TEXT("Primitive Float Get"), Test.DefaultStruct.Float, Test.ModifiedStruct.Float);
+		GPassing &= TestEqual(TEXT("Primitive Float Get Cpp Equivalent"), CppTest.ModifiedStruct, Test.ModifiedStruct);
+	 	GPassing &= TestEqual(TEXT("Primitive Float Setter Called"), Test.Object->IsSetterCalled(), false);
+	 	GPassing &= TestEqual(TEXT("Primitive Float Getter Called"), Test.Object->IsGetterCalled(), true);
+	 	Test.Object->ResetGetterSetterFlags();
+	}
+
+	 return GPassing;
+ }
+
+
+IMPLEMENT_SIMPLE_AUTOMATION_TEST(FGetStringTest, "System.PropertyPath.GetStringTest", EAutomationTestFlags::ApplicationContextMask | EAutomationTestFlags::EngineFilter)
+bool FGetStringTest::RunTest(const FString& Parameters)
+{
+	bool GPassing = true;
+
+	{
+		FPropertyPathTestBed Test = {};
+		FPropertyPathTestBed CppTest = {};
+		CppTest.ModifiedStruct.String = CppTest.Object->String;
+
+		PropertyPathHelpers::GetPropertyValue(Test.Object, FString("String"), Test.ModifiedStruct.String);
+		GPassing &= TestEqual(TEXT("String Get"), Test.DefaultStruct.String, Test.ModifiedStruct.String);
+		GPassing &= TestEqual(TEXT("String Get Cpp Equivalent"), CppTest.ModifiedStruct, Test.ModifiedStruct);
+		GPassing &= TestEqual(TEXT("String Setter Called"), Test.Object->IsSetterCalled(), false);
+		GPassing &= TestEqual(TEXT("String Getter Called"), Test.Object->IsGetterCalled(), false);
+		Test.Object->ResetGetterSetterFlags();
+	}
+
+	return GPassing;
+}
+
+IMPLEMENT_SIMPLE_AUTOMATION_TEST(FGetEnumTest, "System.PropertyPath.GetEnumTest", EAutomationTestFlags::ApplicationContextMask | EAutomationTestFlags::EngineFilter)
+bool FGetEnumTest::RunTest(const FString& Parameters)
+{
+	bool GPassing = true;
+
+	{
+		FPropertyPathTestBed Test = {};
+		FPropertyPathTestBed CppTest = {};
+		CppTest.ModifiedStruct.EnumTwo = CppTest.Object->EnumTwo;
+
+		PropertyPathHelpers::GetPropertyValue(Test.Object, FString("EnumTwo"), Test.ModifiedStruct.EnumTwo);
+		GPassing &= TestEqual(TEXT("Enum Get"), Test.DefaultStruct.EnumTwo, Test.ModifiedStruct.EnumTwo);
+		GPassing &= TestEqual(TEXT("Enum Get Cpp Equivalent"), CppTest.ModifiedStruct, Test.ModifiedStruct);
+		GPassing &= TestEqual(TEXT("Enum Setter Called"), Test.Object->IsSetterCalled(), false);
+		GPassing &= TestEqual(TEXT("Enum Getter Called"), Test.Object->IsGetterCalled(), false);
+		Test.Object->ResetGetterSetterFlags();
+	}
+
+	return GPassing;
+}
+
+IMPLEMENT_SIMPLE_AUTOMATION_TEST(FGetStructTest, "System.PropertyPath.GetStructTest", EAutomationTestFlags::ApplicationContextMask | EAutomationTestFlags::EngineFilter)
+bool FGetStructTest::RunTest(const FString& Parameters)
+{
+	bool GPassing = true;
+
+	{
+		FPropertyPathTestBed Test = {};
+		FPropertyPathTestBed CppTest = {};
+		CppTest.ModifiedStruct = CppTest.Object->GetStruct();
+
+		PropertyPathHelpers::GetPropertyValue(Test.Object, FString("Struct"), Test.ModifiedStruct);
+		GPassing &= TestEqual(TEXT("Struct Get"), Test.DefaultStruct, Test.ModifiedStruct);
+		GPassing &= TestEqual(TEXT("Struct Get Cpp Equivalent"), CppTest.ModifiedStruct, Test.ModifiedStruct);
+		GPassing &= TestEqual(TEXT("Struct Setter Called"), Test.Object->IsSetterCalled(), false);
+		GPassing &= TestEqual(TEXT("Struct Getter Called"), Test.Object->IsGetterCalled(), true);
+		Test.Object->ResetGetterSetterFlags();
+	}
+
+	{
+		FPropertyPathTestBed Test = {};
+		FPropertyPathTestBed CppTest = {};
+		CppTest.ModifiedStruct = CppTest.Object->GetStructRef();
+
+		PropertyPathHelpers::GetPropertyValue(Test.Object, FString("StructRef"), Test.ModifiedStruct);
+		GPassing &= TestEqual(TEXT("Struct Ref Get"), Test.DefaultStruct, Test.ModifiedStruct);
+		GPassing &= TestEqual(TEXT("Struct Ref Get Cpp Equivalent"), CppTest.ModifiedStruct, Test.ModifiedStruct);
+		GPassing &= TestEqual(TEXT("Struct Setter Called"), Test.Object->IsSetterCalled(), false);
+		GPassing &= TestEqual(TEXT("Struct Getter Called"), Test.Object->IsGetterCalled(), true);
+		Test.Object->ResetGetterSetterFlags();
+	}
+
+	{
+		FPropertyPathTestBed Test = {};
+		FPropertyPathTestBed CppTest = {};
+		CppTest.ModifiedStruct = CppTest.Object->GetStructConstRef();
+
+		PropertyPathHelpers::GetPropertyValue(Test.Object, FString("StructConstRef"), Test.ModifiedStruct);
+		GPassing &= TestEqual(TEXT("Struct Const Ref Get"), Test.DefaultStruct, Test.ModifiedStruct);
+		GPassing &= TestEqual(TEXT("Struct Const Ref Get Cpp Equivalent"), CppTest.ModifiedStruct, Test.ModifiedStruct);
+		GPassing &= TestEqual(TEXT("Struct Setter Called"), Test.Object->IsSetterCalled(), false);
+		GPassing &= TestEqual(TEXT("Struct Getter Called"), Test.Object->IsGetterCalled(), true);
+		Test.Object->ResetGetterSetterFlags();
+	}
+
+	return GPassing;
+}
+
+IMPLEMENT_SIMPLE_AUTOMATION_TEST(FGetStructMemberTest, "System.PropertyPath.GetStructMemberTest", EAutomationTestFlags::ApplicationContextMask | EAutomationTestFlags::EngineFilter)
+bool FGetStructMemberTest::RunTest(const FString& Parameters)
+{
+	bool GPassing = true;
+
+	{
+		FPropertyPathTestBed Test = {};
+		FPropertyPathTestBed CppTest = {};
+		CppTest.ModifiedStruct.Float = CppTest.Object->GetStruct().Float;
+
+		PropertyPathHelpers::GetPropertyValue(Test.Object, FString("Struct.Float"), Test.ModifiedStruct.Float);
+		GPassing &= TestEqual(TEXT("Struct Member Get"), Test.DefaultStruct.Float, Test.ModifiedStruct.Float);
+		GPassing &= TestEqual(TEXT("Struct Member Get Cpp Equivalent"), CppTest.ModifiedStruct, Test.ModifiedStruct);
+		GPassing &= TestEqual(TEXT("Struct Setter Called"), Test.Object->IsSetterCalled(), false);
+		GPassing &= TestEqual(TEXT("Struct Getter Called"), Test.Object->IsGetterCalled(), true);
+
+		Test.Object->ResetGetterSetterFlags();
+	}
+
+	{
+		FPropertyPathTestBed Test = {};
+		FPropertyPathTestBed CppTest = {};
+		CppTest.ModifiedStruct.Float = CppTest.Object->GetStructRef().Float;
+
+		PropertyPathHelpers::GetPropertyValue(Test.Object, FString("StructRef.Float"), Test.ModifiedStruct.Float);
+		GPassing &= TestEqual(TEXT("Struct Ref Member Get"), Test.DefaultStruct.Float, Test.ModifiedStruct.Float);
+		GPassing &= TestEqual(TEXT("Struct Ref Member Get Cpp Equivalent"), CppTest.ModifiedStruct, Test.ModifiedStruct);
+		GPassing &= TestEqual(TEXT("Struct Setter Called"), Test.Object->IsSetterCalled(), false);
+		GPassing &= TestEqual(TEXT("Struct Getter Called"), Test.Object->IsGetterCalled(), true);
+
+		Test.Object->ResetGetterSetterFlags();
+	}
+
+	{
+		FPropertyPathTestBed Test = {};
+		FPropertyPathTestBed CppTest = {};
+		CppTest.ModifiedStruct.Float = CppTest.Object->GetStructConstRef().Float;
+
+		PropertyPathHelpers::GetPropertyValue(Test.Object, FString("StructConstRef.Float"), Test.ModifiedStruct.Float);
+		GPassing &= TestEqual(TEXT("Struct Const Ref Member Get"), Test.DefaultStruct.Float, Test.ModifiedStruct.Float);
+		GPassing &= TestEqual(TEXT("Struct Const Ref Member Get Cpp Equivalent"), CppTest.ModifiedStruct, Test.ModifiedStruct);
+		GPassing &= TestEqual(TEXT("Struct Setter Called"), Test.Object->IsSetterCalled(), false);
+		GPassing &= TestEqual(TEXT("Struct Getter Called"), Test.Object->IsGetterCalled(), true);
+
+		Test.Object->ResetGetterSetterFlags();
+	}
+
+	return GPassing;
+}
+
+IMPLEMENT_SIMPLE_AUTOMATION_TEST(FGetInnerStructMemberTest, "System.PropertyPath.GetInnerStructMemberTest", EAutomationTestFlags::ApplicationContextMask | EAutomationTestFlags::EngineFilter)
+bool FGetInnerStructMemberTest::RunTest(const FString& Parameters)
+{
+	bool GPassing = true;
+
+	{
+		FPropertyPathTestBed Test = {};
+		FPropertyPathTestBed CppTest = {};
+		CppTest.ModifiedStruct.Float = CppTest.Object->GetStruct().InnerStruct.Float;
+
+		PropertyPathHelpers::GetPropertyValue(Test.Object, FString("Struct.InnerStruct.Float"), Test.ModifiedStruct.Float);
+		GPassing &= TestEqual(TEXT("Inner Struct Member Get"), Test.DefaultStruct.InnerStruct.Float, Test.ModifiedStruct.Float);
+		GPassing &= TestEqual(TEXT("Inner Struct Member Get Cpp Equivalent"), CppTest.ModifiedStruct, Test.ModifiedStruct);
+		GPassing &= TestEqual(TEXT("Inner Struct Setter Called"), Test.Object->IsSetterCalled(), false);
+		GPassing &= TestEqual(TEXT("Inner Struct Getter Called"), Test.Object->IsGetterCalled(), true);
+
+		Test.Object->ResetGetterSetterFlags();
+	}
+
+	return GPassing;
+}
+
+IMPLEMENT_SIMPLE_AUTOMATION_TEST(FGetPrimitivesInnerObjectTest, "System.PropertyPath.GetPrimitivesInnerObjectTest", EAutomationTestFlags::ApplicationContextMask | EAutomationTestFlags::EngineFilter)
+bool FGetPrimitivesInnerObjectTest::RunTest(const FString& Parameters)
+{
+	bool GPassing = true;
+
+	{
+		FPropertyPathTestBed Test = {};
+		FPropertyPathTestBed CppTest = {};
+		CppTest.ModifiedStruct.Bool = CppTest.Object->InnerObject->Bool;
+
+		PropertyPathHelpers::GetPropertyValue(Test.Object, FString("InnerObject.Bool"), Test.ModifiedStruct.Bool);
+		GPassing &= TestEqual(TEXT("Primitive Inner Object Bool Get"), Test.DefaultStruct.Bool, Test.ModifiedStruct.Bool);
+		GPassing &= TestEqual(TEXT("Primitive Inner Object Bool Get Cpp Equivalent"), CppTest.ModifiedStruct, Test.ModifiedStruct);
+		GPassing &= TestEqual(TEXT("Primitive Inner Object Bool Setter Called"), Test.Object->InnerObject->IsSetterCalled(), false);
+		GPassing &= TestEqual(TEXT("Primitive Inner Object Bool Getter Called"), Test.Object->InnerObject->IsGetterCalled(), false);
+		GPassing &= TestEqual(TEXT("Primitive Bool Setter Called"), Test.Object->IsSetterCalled(), false);
+		GPassing &= TestEqual(TEXT("Primitive Bool Getter Called"), Test.Object->IsGetterCalled(), false);
+		Test.Object->InnerObject->ResetGetterSetterFlags();
+		Test.Object->ResetGetterSetterFlags();
+	}
+
+	{
+		FPropertyPathTestBed Test = {};
+		FPropertyPathTestBed CppTest = {};
+		CppTest.ModifiedStruct.Integer = CppTest.Object->InnerObject->Integer;
+
+		PropertyPathHelpers::GetPropertyValue(Test.Object, FString("InnerObject.Integer"), Test.ModifiedStruct.Integer);
+		GPassing &= TestEqual(TEXT("Primitive Inner Object Integer Get"), Test.DefaultStruct.Integer, Test.ModifiedStruct.Integer);
+		GPassing &= TestEqual(TEXT("Primitive Inner Object Integer Get Cpp Equivalent"), CppTest.ModifiedStruct, Test.ModifiedStruct);
+		GPassing &= TestEqual(TEXT("Primitive Inner Object Integer Setter Called"), Test.Object->InnerObject->IsSetterCalled(), false);
+		GPassing &= TestEqual(TEXT("Primitive Inner Object Integer Getter Called"), Test.Object->InnerObject->IsGetterCalled(), false);
+		GPassing &= TestEqual(TEXT("Primitive Bool Setter Called"), Test.Object->IsSetterCalled(), false);
+		GPassing &= TestEqual(TEXT("Primitive Bool Getter Called"), Test.Object->IsGetterCalled(), false);
+		Test.Object->InnerObject->ResetGetterSetterFlags();
+		Test.Object->ResetGetterSetterFlags();
+	}
+
+	{
+		FPropertyPathTestBed Test = {};
+		FPropertyPathTestBed CppTest = {};
+		CppTest.ModifiedStruct.Float = CppTest.Object->InnerObject->GetFloat();
+
+		PropertyPathHelpers::GetPropertyValue(Test.Object, FString("InnerObject.Float"), Test.ModifiedStruct.Float);
+		GPassing &= TestEqual(TEXT("Primitive Inner Object Float Get"), Test.DefaultStruct.Float, Test.ModifiedStruct.Float);
+		GPassing &= TestEqual(TEXT("Primitive Inner Object Float Get Cpp Equivalent"), CppTest.ModifiedStruct, Test.ModifiedStruct);
+		GPassing &= TestEqual(TEXT("Primitive Inner Object Float Setter Called"), Test.Object->InnerObject->IsSetterCalled(), false);
+		GPassing &= TestEqual(TEXT("Primitive Inner Object Float Getter Called"), Test.Object->InnerObject->IsGetterCalled(), true);
+		GPassing &= TestEqual(TEXT("Primitive Float Setter Called"), Test.Object->IsSetterCalled(), false);
+		GPassing &= TestEqual(TEXT("Primitive Float Getter Called"), Test.Object->IsGetterCalled(), false);
+		Test.Object->InnerObject->ResetGetterSetterFlags();
+		Test.Object->ResetGetterSetterFlags();
+	}
+
+	return GPassing;
+}
+
+IMPLEMENT_SIMPLE_AUTOMATION_TEST(FGetStringInnerObjectTest, "System.PropertyPath.GetStringInnerObjectTest", EAutomationTestFlags::ApplicationContextMask | EAutomationTestFlags::EngineFilter)
+bool FGetStringInnerObjectTest::RunTest(const FString& Parameters)
+{
+	bool GPassing = true;
+
+	{
+		FPropertyPathTestBed Test = {};
+		FPropertyPathTestBed CppTest = {};
+		CppTest.ModifiedStruct.String = CppTest.Object->InnerObject->String;
+
+		PropertyPathHelpers::GetPropertyValue(Test.Object, FString("InnerObject.String"), Test.ModifiedStruct.String);
+		GPassing &= TestEqual(TEXT("Inner Object String Get"), Test.DefaultStruct.String, Test.ModifiedStruct.String);
+		GPassing &= TestEqual(TEXT("Inner Object String Get Cpp Equivalent"), CppTest.ModifiedStruct, Test.ModifiedStruct);
+		GPassing &= TestEqual(TEXT("Inner Object String Setter Called"), Test.Object->InnerObject->IsSetterCalled(), false);
+		GPassing &= TestEqual(TEXT("Inner Object String Getter Called"), Test.Object->InnerObject->IsGetterCalled(), false);
+		GPassing &= TestEqual(TEXT("String Setter Called"), Test.Object->IsSetterCalled(), false);
+		GPassing &= TestEqual(TEXT("String Getter Called"), Test.Object->IsGetterCalled(), false);
+		Test.Object->InnerObject->ResetGetterSetterFlags();
+		Test.Object->ResetGetterSetterFlags();
+	}
+
+	return GPassing;
+}
+
+IMPLEMENT_SIMPLE_AUTOMATION_TEST(FGetStructInnerObjectTest, "System.PropertyPath.GetStructInnerObjectTest", EAutomationTestFlags::ApplicationContextMask | EAutomationTestFlags::EngineFilter)
+bool FGetStructInnerObjectTest::RunTest(const FString& Parameters)
+{
+	bool GPassing = true;
+
+	{
+		FPropertyPathTestBed Test = {};
+		FPropertyPathTestBed CppTest = {};
+		CppTest.ModifiedStruct = CppTest.Object->InnerObject->GetStruct();
+
+		PropertyPathHelpers::GetPropertyValue(Test.Object, FString("InnerObject.Struct"), Test.ModifiedStruct);
+		GPassing &= TestEqual(TEXT("Inner Object Struct Get"), Test.DefaultStruct, Test.ModifiedStruct);
+		GPassing &= TestEqual(TEXT("Inner Object Struct Get Cpp Equivalent"), CppTest.ModifiedStruct, Test.ModifiedStruct);
+		GPassing &= TestEqual(TEXT("Inner Object Struct Setter Called"), Test.Object->InnerObject->IsSetterCalled(), false);
+		GPassing &= TestEqual(TEXT("Inner Object Struct Getter Called"), Test.Object->InnerObject->IsGetterCalled(), true);
+		GPassing &= TestEqual(TEXT("Struct Setter Called"), Test.Object->IsSetterCalled(), false);
+		GPassing &= TestEqual(TEXT("Struct Getter Called"), Test.Object->IsGetterCalled(), false);
+		Test.Object->InnerObject->ResetGetterSetterFlags();
+		Test.Object->ResetGetterSetterFlags();
+	}
+	return GPassing;
+}
+
+IMPLEMENT_SIMPLE_AUTOMATION_TEST(FGetStructMemberInnerObjectTest, "System.PropertyPath.GetStructMemberInnerObjectTest", EAutomationTestFlags::ApplicationContextMask | EAutomationTestFlags::EngineFilter)
+bool FGetStructMemberInnerObjectTest::RunTest(const FString& Parameters)
+{
+	bool GPassing = true;
+
+	{
+		FPropertyPathTestBed Test = {};
+		FPropertyPathTestBed CppTest = {};
+		CppTest.ModifiedStruct.Float = CppTest.Object->InnerObject->GetStruct().Float;
+
+		PropertyPathHelpers::GetPropertyValue(Test.Object, FString("InnerObject.Struct.Float"), Test.ModifiedStruct.Float);
+		GPassing &= TestEqual(TEXT("Inner Object Struct Member Get"), Test.DefaultStruct.Float, Test.ModifiedStruct.Float);
+		GPassing &= TestEqual(TEXT("Inner Object Struct Member Get Cpp Equivalent"), CppTest.ModifiedStruct, Test.ModifiedStruct);
+		GPassing &= TestEqual(TEXT("Inner Object Struct Setter Called"), Test.Object->InnerObject->IsSetterCalled(), false);
+		GPassing &= TestEqual(TEXT("Inner Object Struct Getter Called"), Test.Object->InnerObject->IsGetterCalled(), true);
+		GPassing &= TestEqual(TEXT("Struct Setter Called"), Test.Object->IsSetterCalled(), false);
+		GPassing &= TestEqual(TEXT("Struct Getter Called"), Test.Object->IsGetterCalled(), false);
+		Test.Object->ResetGetterSetterFlags();
+	}
+
+	return GPassing;
+}
+
+IMPLEMENT_SIMPLE_AUTOMATION_TEST(FGetInnerStructMemberInnerObjectTest, "System.PropertyPath.GetInnerStructMemberInnerObjectTest", EAutomationTestFlags::ApplicationContextMask | EAutomationTestFlags::EngineFilter)
+bool FGetInnerStructMemberInnerObjectTest::RunTest(const FString& Parameters)
+{
+	bool GPassing = true;
+
+	{
+		FPropertyPathTestBed Test = {};
+		FPropertyPathTestBed CppTest = {};
+		CppTest.ModifiedStruct.Float = CppTest.Object->InnerObject->GetStruct().InnerStruct.Float;
+
+		PropertyPathHelpers::GetPropertyValue(Test.Object, FString("InnerObject.Struct.InnerStruct.Float"), Test.ModifiedStruct.Float);
+		GPassing &= TestEqual(TEXT("Inner Object Inner Struct Member Get"), Test.DefaultStruct.InnerStruct.Float, Test.ModifiedStruct.Float);
+		GPassing &= TestEqual(TEXT("Inner Object Inner Struct Member Get Cpp Equivalent"), CppTest.ModifiedStruct, Test.ModifiedStruct);
+		GPassing &= TestEqual(TEXT("Inner Object Inner Struct Setter Called"), Test.Object->InnerObject->IsSetterCalled(), false);
+		GPassing &= TestEqual(TEXT("Inner Object Inner Struct Getter Called"), Test.Object->InnerObject->IsGetterCalled(), true);
+		GPassing &= TestEqual(TEXT("Inner Struct Setter Called"), Test.Object->IsSetterCalled(), false);
+		GPassing &= TestEqual(TEXT("Inner Struct Getter Called"), Test.Object->IsGetterCalled(), false);
+		Test.Object->ResetGetterSetterFlags();
+	}
+
+	return GPassing;
+}
+
+IMPLEMENT_SIMPLE_AUTOMATION_TEST(FGetInnerStructMemberEnumInnerObjectTest, "System.PropertyPath.GetInnerStructMemberEnumInnerObjectTest", EAutomationTestFlags::ApplicationContextMask | EAutomationTestFlags::EngineFilter)
+bool FGetInnerStructMemberEnumInnerObjectTest::RunTest(const FString& Parameters)
+{
+	bool GPassing = true;
+
+	{
+		FPropertyPathTestBed Test = {};
+		FPropertyPathTestBed CppTest = {};
+		CppTest.ModifiedStruct.EnumTwo = CppTest.Object->InnerObject->GetStruct().InnerStruct.EnumTwo;
+		PropertyPathHelpers::GetPropertyValue(Test.Object, FString("InnerObject.Struct.InnerStruct.EnumTwo"), Test.ModifiedStruct.EnumTwo);
+		GPassing &= TestEqual(TEXT("Inner Object Inner Struct Member Enum Get"), Test.DefaultStruct.InnerStruct.EnumTwo, Test.ModifiedStruct.EnumTwo);
+		GPassing &= TestEqual(TEXT("Inner Object Inner Struct Member Enum Get Cpp Equivalent"), CppTest.ModifiedStruct, Test.ModifiedStruct);
+		GPassing &= TestEqual(TEXT("Inner Object Inner Struct Setter Called"), Test.Object->InnerObject->IsSetterCalled(), false);
+		GPassing &= TestEqual(TEXT("Inner Object Inner Struct Getter Called"), Test.Object->InnerObject->IsGetterCalled(), true);
+		GPassing &= TestEqual(TEXT("Inner Struct Setter Called"), Test.Object->IsSetterCalled(), false);
+		GPassing &= TestEqual(TEXT("Inner Struct Getter Called"), Test.Object->IsGetterCalled(), false);
+		Test.Object->ResetGetterSetterFlags();
+	}
+
+	return GPassing;
+}
+
+#endif // WITH_DEV_AUTOMATION_TESTS
\ No newline at end of file
diff --git a/Engine/Source/Runtime/PropertyPath/Private/Tests/PropertyPathHelpersTest.h b/Engine/Source/Runtime/PropertyPath/Private/Tests/PropertyPathHelpersTest.h
new file mode 100644
--- /dev/null
+++ b/Engine/Source/Runtime/PropertyPath/Private/Tests/PropertyPathHelpersTest.h
@@ -0,0 +1,344 @@
+// Copyright Epic Games, Inc. All Rights Reserved.
+
+#pragma once
+
+#include "CoreTypes.h"
+#include "Logging/LogMacros.h"
+#include "Misc/AutomationTest.h"
+#include "PropertyPathHelpers.h"
+#include "PropertyPathHelpersTest.generated.h"
+
+UENUM()
+enum EPropertyPathTestEnum : int
+{
+	One,
+	Two,
+	Three,
+	Four
+};
+
+
+USTRUCT(BlueprintType)
+struct FPropertyPathTestInnerStruct
+{
+	GENERATED_BODY()
+
+	~FPropertyPathTestInnerStruct()
+	{
+		Float = -1.5f;
+		Bool = true;
+		EnumOne = Four;
+		EnumTwo = Four;
+		EnumThree = Four;
+		EnumFour = Four;
+		Integer = -1;
+		String = "Destructed";
+	}
+
+	UPROPERTY()
+	float Float = 0.5f;
+
+	UPROPERTY()
+	bool Bool = false;
+
+	UPROPERTY()
+	TEnumAsByte<EPropertyPathTestEnum> EnumOne = Three;
+
+	UPROPERTY()
+	TEnumAsByte<EPropertyPathTestEnum> EnumTwo = Three;
+
+	UPROPERTY()
+	TEnumAsByte<EPropertyPathTestEnum> EnumThree = Three;
+
+	UPROPERTY()
+	TEnumAsByte<EPropertyPathTestEnum> EnumFour = Three;
+
+	UPROPERTY()
+	int32 Integer = 0;
+
+	UPROPERTY()
+	FString String = "Default";
+
+	bool operator ==(const FPropertyPathTestInnerStruct& Other) const
+	{
+		return Bool == Other.Bool
+			&& Integer == Other.Integer
+			&& EnumOne == Other.EnumOne
+			&& EnumTwo == Other.EnumTwo
+			&& EnumThree == Other.EnumThree
+			&& EnumFour == Other.EnumFour
+			&& String == Other.String
+			&& Float == Other.Float;
+	}
+
+	bool operator !=(const FPropertyPathTestInnerStruct& Other) const
+	{
+		return !(*this == Other);
+	}
+};
+
+USTRUCT(BlueprintType)
+struct FPropertyPathTestStruct
+{
+	GENERATED_BODY()
+
+	~FPropertyPathTestStruct()
+	{
+		Float = -1.5f;
+		Bool = true;
+		EnumOne = Two;
+		EnumTwo = Two;
+		EnumThree = Two;
+		EnumFour = Two;
+		Integer = -1;
+		String = "Destructed";
+	}
+
+	UPROPERTY()
+	bool Bool = false;
+
+	UPROPERTY()
+	int32 Integer = 0;
+
+	UPROPERTY()
+	TEnumAsByte<EPropertyPathTestEnum> EnumOne = One;
+
+	UPROPERTY()
+	TEnumAsByte<EPropertyPathTestEnum> EnumTwo = One;
+
+	UPROPERTY()
+	TEnumAsByte<EPropertyPathTestEnum> EnumThree = One;
+
+	UPROPERTY()
+	TEnumAsByte<EPropertyPathTestEnum> EnumFour = One;
+
+	UPROPERTY()
+	FString String = "Default";
+
+	UPROPERTY()
+	float Float = 0.0f;
+
+	UPROPERTY()
+	FPropertyPathTestInnerStruct InnerStruct = {};
+
+	bool operator ==(const FPropertyPathTestStruct& Other) const
+	{
+		return Bool == Other.Bool
+			&& Integer == Other.Integer
+			&& EnumOne == Other.EnumOne
+			&& EnumTwo == Other.EnumTwo
+			&& EnumThree == Other.EnumThree
+			&& EnumFour == Other.EnumFour
+			&& String == Other.String
+			&& Float == Other.Float
+			&& InnerStruct == Other.InnerStruct;
+	}
+
+	bool operator !=(const FPropertyPathTestStruct& Other) const
+	{
+		return !(*this == Other);
+	}
+};
+
+UCLASS()
+class UPropertyPathTestObject : public UObject
+{
+	GENERATED_BODY()
+		
+public:
+
+	~UPropertyPathTestObject()
+	{
+		Float = -1.5f;
+		Bool = true;
+		EnumOne = Two;
+		EnumTwo = Two;
+		EnumThree = Two;
+		EnumFour = Two;
+		Integer = -1;
+		String = "Destructed";
+	}
+
+	UPROPERTY(EditDefaultsOnly, BlueprintReadWrite, Category = "PropertyPathHelpersTest")
+	bool Bool = false;
+
+	UPROPERTY(EditDefaultsOnly, Category = "PropertyPathHelpersTest")
+	TEnumAsByte<EPropertyPathTestEnum> EnumOne = One;
+
+	UPROPERTY(EditDefaultsOnly, Category = "PropertyPathHelpersTest")
+	TEnumAsByte<EPropertyPathTestEnum> EnumTwo = One;
+
+	UPROPERTY(EditDefaultsOnly, Category = "PropertyPathHelpersTest")
+	TEnumAsByte<EPropertyPathTestEnum> EnumThree = One;
+
+	UPROPERTY(EditDefaultsOnly, Category = "PropertyPathHelpersTest")
+	TEnumAsByte<EPropertyPathTestEnum> EnumFour = One;
+
+	UPROPERTY(EditDefaultsOnly, BlueprintReadWrite, Category = "PropertyPathHelpersTest")
+	int32 Integer = 0;
+
+	UPROPERTY(EditDefaultsOnly, BlueprintReadWrite, Category = "PropertyPathHelpersTest")
+	FString String = "Default";
+
+	UPROPERTY(EditDefaultsOnly, BlueprintReadWrite, Getter, Setter, Category = "PropertyPathHelpersTest")
+	float Float = 0.0f;
+
+	UPROPERTY(EditDefaultsOnly, BlueprintReadWrite, Getter, Setter, Category = "PropertyPathHelpersTest")
+	FPropertyPathTestStruct Struct;
+
+	UPROPERTY(EditDefaultsOnly, BlueprintReadWrite, Getter, Setter, Category = "PropertyPathHelpersTest")
+	mutable FPropertyPathTestStruct StructRef;
+
+	UPROPERTY(EditDefaultsOnly, BlueprintReadWrite, Getter, Setter, Category = "PropertyPathHelpersTest")
+	FPropertyPathTestStruct StructConstRef;
+
+	UPROPERTY(EditAnywhere, BlueprintReadOnly, Category = "PropertyPathHelpersTest")
+	TObjectPtr<UPropertyPathTestObject> InnerObject;
+
+public:
+
+	UFUNCTION()
+	void SetFloat(float InFloat)
+	{
+		bSetterCalled = true;
+		Float = InFloat;
+	}
+
+	UFUNCTION()
+	float GetFloat() const
+	{
+		bGetterCalled = true;
+		return Float;
+	}
+
+	UFUNCTION()
+	void SetStruct(FPropertyPathTestStruct InStruct)
+	{
+		bSetterCalled = true;
+		Struct = InStruct;
+	}
+
+	UFUNCTION()
+	FPropertyPathTestStruct GetStruct() const
+	{
+		bGetterCalled = true;
+		return Struct;
+	}
+
+	UFUNCTION()
+	void SetStructRef(FPropertyPathTestStruct InStruct)
+	{
+		bSetterCalled = true;
+		StructRef = InStruct;
+	}
+
+	UFUNCTION()
+	FPropertyPathTestStruct& GetStructRef() const
+	{
+		bGetterCalled = true;
+		return StructRef;
+	}
+
+	UFUNCTION()
+	void SetStructConstRef(FPropertyPathTestStruct InStruct)
+	{
+		bSetterCalled = true;
+		StructConstRef = InStruct;
+	}
+
+	UFUNCTION()
+	const FPropertyPathTestStruct& GetStructConstRef() const
+	{
+		bGetterCalled = true;
+		return StructConstRef;
+	}
+
+public:
+
+	bool operator ==(const UPropertyPathTestObject& Other) const
+	{
+		return Bool == Other.Bool
+			&& Integer == Other.Integer
+			&& EnumOne == Other.EnumOne
+			&& EnumTwo == Other.EnumTwo
+			&& EnumThree == Other.EnumThree
+			&& EnumFour == Other.EnumFour
+			&& String == Other.String
+			&& Float == Other.Float
+			&& Struct == Other.Struct
+			&& StructRef == Other.StructRef
+			&& StructConstRef == Other.StructConstRef
+			&& bSetterCalled == Other.bSetterCalled
+			&& bGetterCalled == Other.bGetterCalled
+			&& InnerObject ? (*InnerObject == *Other.InnerObject) : true;
+	}
+
+	bool operator !=(const UPropertyPathTestObject& Other) const
+	{
+		return !(*this == Other);
+	}
+
+	bool IsSetterCalled() const
+	{
+		return bSetterCalled;
+	}
+
+	bool IsGetterCalled() const
+	{
+		return bGetterCalled;
+	}
+		
+	void ResetGetterSetterFlags() const
+	{
+		bSetterCalled = false;
+		bGetterCalled = false;
+	}
+
+private:
+
+	mutable bool bSetterCalled = false;
+	mutable bool bGetterCalled = false;
+};
+
+USTRUCT()
+struct FPropertyPathTestBed
+{
+	GENERATED_BODY()
+
+	FPropertyPathTestBed()
+	{
+		Object = NewObject<UPropertyPathTestObject>();
+		Object->InnerObject = NewObject<UPropertyPathTestObject>();
+
+		ModifiedStruct = {};
+
+		ModifiedStruct.Bool = true;
+		ModifiedStruct.Integer = 1;
+		ModifiedStruct.EnumOne = Two;
+		ModifiedStruct.EnumTwo = Two;
+		ModifiedStruct.EnumThree = Two;
+		ModifiedStruct.EnumFour = Two;
+		ModifiedStruct.String = "NewValue";
+		ModifiedStruct.Float = 1.5f;
+
+		ModifiedStruct.InnerStruct.Bool = true;
+		ModifiedStruct.InnerStruct.Integer = 1;
+		ModifiedStruct.EnumOne = Four;
+		ModifiedStruct.EnumTwo = Four;
+		ModifiedStruct.EnumThree = Four;
+		ModifiedStruct.EnumFour = Four;
+		ModifiedStruct.InnerStruct.String = "NewValue";
+		ModifiedStruct.InnerStruct.Float = 1.5f;
+
+		DefaultStruct = {};
+	}
+
+	UPROPERTY()
+	TObjectPtr<UPropertyPathTestObject> Object;
+
+	UPROPERTY()
+	FPropertyPathTestStruct ModifiedStruct;
+
+	UPROPERTY()
+	FPropertyPathTestStruct DefaultStruct;
+};
\ No newline at end of file
diff --git a/Engine/Source/Runtime/PropertyPath/PropertyPath.Build.cs b/Engine/Source/Runtime/PropertyPath/PropertyPath.Build.cs
--- a/Engine/Source/Runtime/PropertyPath/PropertyPath.Build.cs
+++ b/Engine/Source/Runtime/PropertyPath/PropertyPath.Build.cs
@@ -11,5 +11,11 @@ public class PropertyPath : ModuleRules
 			"Core",
 			"CoreUObject" 
 		});
+
+		PrivateIncludePaths.AddRange(
+			new string[] {
+				"Runtime/PropertyPath/Private"
+			}
+		);
 	}
 }
diff --git a/Engine/Source/Runtime/PropertyPath/Public/PropertyPathHelpers.h b/Engine/Source/Runtime/PropertyPath/Public/PropertyPathHelpers.h
--- a/Engine/Source/Runtime/PropertyPath/Public/PropertyPathHelpers.h
+++ b/Engine/Source/Runtime/PropertyPath/Public/PropertyPathHelpers.h
@@ -143,6 +143,15 @@ struct PROPERTYPATH_API FCachedPropertyPath
 	/** Set whether this path resolves over object or dynamic array boundaries, making it unsafe for general direct cached access */
 	void SetCanSafelyUsedCachedAddress(bool bInCanSafelyUsedCachedAddress) const;
 
+	/** Update cached last container property in path & correspondng index, invalidates 'bCanSafelyUsedCachedAddress' */
+	void SetCachedLastContainer(void* InContainer, int32 InIndex) const;
+
+	/** Get cached last container property in path */
+	void* GetCachedLastContainerInPath() const;
+
+	/** Get cached index of last container property in path, INDEX_NONE if last container is not in path */
+	int32 GetCachedLastContainerInPathIndex() const;
+
 	/** Cache a resolved address for faster subsequent access */
 	void ResolveLeaf(void* InAddress) const;
 
@@ -184,13 +193,11 @@ struct PROPERTYPATH_API FCachedPropertyPath
 	/** Compares this property path to a string */
 	bool Equals(const FString& Other) const;
 
-#if DO_CHECK || USING_CODE_ANALYSIS
 	/** Get the cached container for this property path, for checking purposes */
 	void* GetCachedContainer() const;
 
 	/** Set the cached container for this property path, for checking purposes */
 	void SetCachedContainer(void* InContainer) const;
-#endif
 
 	/** Trims this property path at the end */
 	void RemoveFromEnd(int32 InNumSegments = 1);
@@ -212,10 +219,14 @@ private:
 	UPROPERTY()
 	mutable TObjectPtr<UFunction> CachedFunction;
 
-#if DO_CHECK || USING_CODE_ANALYSIS
-	/** Cached container, used for checking purposes only */
+	/** Cached container */
 	mutable void* CachedContainer;
-#endif
+
+	/** Cached last container */
+	mutable void* CachedLastContainerInPath;
+
+	/** Index of last container property in path, INDEX_NONE if last container is not in path */
+	mutable int32 CachedLastContainerInPathIndex;
 
 	/** Whether this path resolves over object or dynamic array boundaries, making it unsafe for general direct cached access */
 	mutable bool bCanSafelyUsedCachedAddress;
@@ -225,6 +236,7 @@ private:
 namespace PropertyPathHelpersInternal
 {
 	struct FPropertyPathResolver;
+	struct FPropertyStructView;
 	template<typename T> struct FInternalGetterResolver;
 	template<typename T> struct FInternalSetterResolver;
 
@@ -519,6 +531,34 @@ namespace PropertyPathHelpers
 		return SetPropertyValue(InContainer, CachedPath, InValue);
 	}
 
+	/** 
+	 * Set the value and the leaf property represented by this property path 
+	 * If the cached property path has a cached address it will use that as a 'fast path' instead 
+	 * of iterating the path. This has safety implications depending on the form of the path, so 
+	 * paths that are resolved over object boundaries or dynamic arrays will always use the slow 
+	 * path for safety.
+	 * @param	InContainer		The container object to resolve the property path against
+	 * @param	InPropertyPath	The property path
+	 * @param	InScriptStruct	The struct type to set
+	 * @param	InValue			A pointer to the desired value for the given struct type
+	 * @return true if the property value was successfully copied
+	 */
+	PROPERTYPATH_API bool SetPropertyValue(UObject* InContainer, const FCachedPropertyPath& InPropertyPath, const UScriptStruct* InScriptStruct, const uint8* InValue);
+
+	/** 
+	 * Set the value and the leaf property represented by this property path 
+	 * If the cached property path has a cached address it will use that as a 'fast path' instead 
+	 * of iterating the path. This has safety implications depending on the form of the path, so 
+	 * paths that are resolved over object boundaries or dynamic arrays will always use the slow 
+	 * path for safety.
+	 * @param	InContainer		The container object to resolve the property path against
+	 * @param	InPropertyPath	The property path string
+	 * @param	InScriptStruct	The struct type to set
+	 * @param	InValue			A pointer to the desired value for the given struct type
+	 * @return true if the property value was successfully copied
+	 */
+	PROPERTYPATH_API bool SetPropertyValue(UObject* InContainer, const FString& InPropertyPath, const UScriptStruct* InScriptStruct, const uint8* InValue);
+
 	/** 
 	 * Copy values between two property paths in the same container.
 	 * @param	InContainer			The container object to resolve the property path against
@@ -588,6 +628,19 @@ namespace PropertyPathHelpersInternal
 		virtual bool Resolve(UObject* InContainer, const FCachedPropertyPath& InPropertyPath) =  0;
 	};
 
+	/** Helper struct to represent a view of a struct, used to specialize behavior against */
+	struct FPropertyStructView
+	{
+		explicit FPropertyStructView(const UScriptStruct* InScriptStruct, const uint8* InMemory)
+			: ScriptStruct(InScriptStruct)
+			, Memory(InMemory)
+		{
+		}
+
+		const UScriptStruct* ScriptStruct;
+		const uint8* Memory;
+	};
+
 	/** Recurring template allowing derived types to only implement templated Resolve_Impl */
 	template<typename DerivedType>
 	struct TPropertyPathResolver : public FPropertyPathResolver
@@ -606,6 +659,16 @@ namespace PropertyPathHelpersInternal
 	/** Find the first param that isnt a return property for the specified function */
 	PROPERTYPATH_API FProperty* GetFirstParamProperty(UFunction* InFunction);
 
+	/** Helper function used to call parent setters, used when modifying elements without setters in a struct that has one */
+	PROPERTYPATH_API void CallParentSetters(const FCachedPropertyPath& InPropertyPath);
+
+	/** 
+	 * Helper function used to call parent getters, used when modifying elements without getters in a struct that has one 
+	 * 
+	 * @return True if a parent getter was called.
+	 */
+	PROPERTYPATH_API void CallParentGetters(void* OutValue, const FCachedPropertyPath& InPropertyPath, const void* InPropertyAddress);
+
 	/** Non-UObject helper struct for GetValue function calls */
 	template<typename T, typename ContainerType>
 	struct FCallGetterFunctionHelper
@@ -660,15 +723,19 @@ namespace PropertyPathHelpersInternal
 			if ( IsConcreteTypeCompatibleWithReflectedType<T>(Property) )
 			{
 				ArrayIndex = ArrayIndex == INDEX_NONE ? 0 : ArrayIndex;
-				if ( PropertySizesMatch<T>(Property) && ArrayIndex < Property->ArrayDim )
+				if (PropertySizesMatch<T>(Property) && ArrayIndex < Property->ArrayDim)
 				{
-					if(void* Address = Property->ContainerPtrToValuePtr<T>(InContainer, ArrayIndex))
+					if (Property->HasGetter())
+					{
+						Property->CallGetter(InContainer, &OutValue);
+					}
+					else if (void* Address = Property->ContainerPtrToValuePtr<T>(InContainer, ArrayIndex))
 					{
 						InPropertyPath.ResolveLeaf(Address);
-						Property->CopySingleValue(&OutValue, Address);
-						OutProperty = Property;
-						return true;
+						CallParentGetters(&OutValue, InPropertyPath, Address);
 					}
+					OutProperty = Property;
+					return true;
 				}
 			}
 
@@ -851,8 +918,19 @@ namespace PropertyPathHelpersInternal
 				// Verify there's a return property.
 				if ( FProperty* ParamProperty = GetFirstParamProperty(InFunction) )
 				{
+					if constexpr (std::is_same<FPropertyStructView, T>::value)
+					{
+						const FPropertyStructView& InStuctView = static_cast<const FPropertyStructView&>(InValue);
+
+						// Ensure that the element sizes are the same, prevents the user from doing something terribly wrong.
+						if (ParamProperty->ElementSize == InStuctView.ScriptStruct->GetStructureSize() && !InContainer->IsUnreachable())
+						{
+							InContainer->ProcessEvent(InFunction, const_cast<uint8*>(InStuctView.Memory));
+							return true;
+						}
+					}
 					// Verify that the cpp type matches a known property type.
-					if ( IsConcreteTypeCompatibleWithReflectedType<T>(ParamProperty) )
+					else if ( IsConcreteTypeCompatibleWithReflectedType<T>(ParamProperty) )
 					{
 						// Ensure that the element sizes are the same, prevents the user from doing something terribly wrong.
 						if ( PropertySizesMatch<T>(ParamProperty) && !InContainer->IsUnreachable() )
@@ -882,19 +960,45 @@ namespace PropertyPathHelpersInternal
 			int32 ArrayIndex = LastSegment.GetArrayIndex();
 			FProperty* Property = CastFieldChecked<FProperty>(LastSegment.GetField().ToField());
 
+			if constexpr (std::is_same<FPropertyStructView, T>::value)
+			{
+				const FPropertyStructView& InStuctView = static_cast<const FPropertyStructView&>(InValue);
+
+				// Ensure that the element sizes are the same, prevents the user from doing something terribly wrong.
+				ArrayIndex = ArrayIndex == INDEX_NONE ? 0 : ArrayIndex;
+				if ( Property->ElementSize == InStuctView.ScriptStruct->GetStructureSize() && ArrayIndex < Property->ArrayDim)
+				{
+					if (Property->HasSetter())
+					{
+						Property->CallSetter(InContainer, InStuctView.Memory);
+					}
+					else if (void* Address = Property->ContainerPtrToValuePtr<void>(InContainer, ArrayIndex))
+					{
+						InPropertyPath.ResolveLeaf(Address);
+						Property->CopySingleValue(Address, InStuctView.Memory);
+						CallParentSetters(InPropertyPath);
+					}
+					return true;
+				}
+			}
 			// Verify that the cpp type matches a known property type.
-			if ( IsConcreteTypeCompatibleWithReflectedType<T>(Property) )
+			else if ( IsConcreteTypeCompatibleWithReflectedType<T>(Property) )
 			{
 				// Ensure that the element sizes are the same, prevents the user from doing something terribly wrong.
 				ArrayIndex = ArrayIndex == INDEX_NONE ? 0 : ArrayIndex;
 				if ( PropertySizesMatch<T>(Property) && ArrayIndex < Property->ArrayDim )
 				{
-					if(void* Address = Property->ContainerPtrToValuePtr<T>(InContainer, ArrayIndex))
+					if (Property->HasSetter())
+					{
+						Property->CallSetter(InContainer, &InValue);
+					}
+					else if (void* Address = Property->ContainerPtrToValuePtr<T>(InContainer, ArrayIndex))
 					{
 						InPropertyPath.ResolveLeaf(Address);
 						Property->CopySingleValue(Address, &InValue);
-						return true;
+						CallParentSetters(InPropertyPath);
 					}
+					return true;
 				}
 			}
 
@@ -950,8 +1054,18 @@ namespace PropertyPathHelpersInternal
 					if(void* Address = Property->ContainerPtrToValuePtr<bool>(InContainer, ArrayIndex))
 					{
 						InPropertyPath.ResolveLeaf(Address);
-						FBoolProperty* BoolProperty = CastFieldChecked<FBoolProperty>(LastSegment.GetField().ToField());
-						BoolProperty->SetPropertyValue(Address, InValue);
+
+						if (Property->HasSetter())
+						{
+							// Setter should be specialized to handle masking
+							Property->CallSetter(InContainer, &InValue);
+						}
+						else
+						{
+							FBoolProperty* BoolProperty = CastFieldChecked<FBoolProperty>(Property);
+							BoolProperty->SetPropertyValue(InPropertyPath.GetCachedAddress(), InValue);
+							CallParentSetters(InPropertyPath);
+						}
 						return true;
 					}
 				}
@@ -967,6 +1081,7 @@ namespace PropertyPathHelpersInternal
 	 * @param InPropertyPath		The property path to set from
 	 * @param InValue				The value to set
 	 * @return true if the value was resolved
+	 * @note This method is specialized internally to work with FPropertyStructView
 	 */
 	template<typename T, typename ContainerType>
 	bool SetValue(ContainerType* InContainer, const FCachedPropertyPath& InPropertyPath, const T& InValue)
@@ -986,36 +1101,57 @@ namespace PropertyPathHelpersInternal
 				FScriptArrayHelper_InContainer ArrayHelper(ArrayProp, InContainer);
 				if ( ArrayHelper.IsValidIndex(ArrayIndex) )
 				{
-					// Verify that the cpp type matches a known property type.
-					if ( IsConcreteTypeCompatibleWithReflectedType<T>(ArrayProp->Inner) )
+					if constexpr (std::is_same<FPropertyStructView, T>::value)
 					{
+						const FPropertyStructView& InStuctView = static_cast<const FPropertyStructView&>(InValue);
+
 						// Ensure that the element sizes are the same, prevents the user from doing something terribly wrong.
-						if ( PropertySizesMatch<T>(ArrayProp->Inner) )
+						if (ArrayProp->Inner->ElementSize == InStuctView.ScriptStruct->GetStructureSize())
 						{
-							if(void* Address = static_cast<void*>(ArrayHelper.GetRawPtr(ArrayIndex)))
+							if (void* Address = static_cast<void*>(ArrayHelper.GetRawPtr(ArrayIndex)))
 							{
 								InPropertyPath.ResolveLeaf(Address);
-								ArrayProp->Inner->CopySingleValue(Address, &InValue);
+								ArrayProp->Inner->CopySingleValue(Address, InStuctView.Memory);
 								return true;
 							}
 						}
 					}
+					else
+					{
+						// Verify that the cpp type matches a known property type.
+						if (IsConcreteTypeCompatibleWithReflectedType<T>(ArrayProp->Inner))
+						{
+							// Ensure that the element sizes are the same, prevents the user from doing something terribly wrong.
+							if (PropertySizesMatch<T>(ArrayProp->Inner))
+							{
+								if (void* Address = static_cast<void*>(ArrayHelper.GetRawPtr(ArrayIndex)))
+								{
+									InPropertyPath.ResolveLeaf(Address);
+									ArrayProp->Inner->CopySingleValue(Address, &InValue);
+									return true;
+								}
+							}
+						}
+					}
 				}
 			}
 			else
 			{
 				// No index, so assume we want the array property itself
 				// Verify that the cpp type matches a known property type.
-				if ( IsConcreteTypeCompatibleWithReflectedType<T>(ArrayProp) )
+				if constexpr (!std::is_same<FPropertyStructView, T>::value)
 				{
-					// Ensure that the element sizes are the same, prevents the user from doing something terribly wrong.
-					if ( PropertySizesMatch<T>(ArrayProp) )
+					if (IsConcreteTypeCompatibleWithReflectedType<T>(ArrayProp))
 					{
-						if(void* Address = ArrayProp->ContainerPtrToValuePtr<T>(InContainer))
+						// Ensure that the element sizes are the same, prevents the user from doing something terribly wrong.
+						if (PropertySizesMatch<T>(ArrayProp))
 						{
-							InPropertyPath.ResolveLeaf(Address);
-							ArrayProp->CopySingleValue(Address, &InValue);
-							return true;
+							if (void* Address = ArrayProp->ContainerPtrToValuePtr<T>(InContainer))
+							{
+								InPropertyPath.ResolveLeaf(Address);
+								ArrayProp->CopySingleValue(Address, &InValue);
+								return true;
+							}
 						}
 					}
 				}
@@ -1108,7 +1244,14 @@ namespace PropertyPathHelpersInternal
 			}
 			else if (IsConcreteTypeCompatibleWithReflectedType<T>(OutProperty))
 			{
-				OutProperty->CopySingleValue(&OutValue, InPropertyPath.GetCachedAddress());
+				if (OutProperty->HasGetter())
+				{
+					OutProperty->CallGetter(InPropertyPath.GetCachedContainer(), &OutValue);
+				}
+				else
+				{
+					CallParentGetters(&OutValue, InPropertyPath, InPropertyPath.GetCachedAddress());
+				}	
 				return true;
 			}
 			return false;
@@ -1132,33 +1275,6 @@ namespace PropertyPathHelpersInternal
 		}
 	};
 
-	/** Explicit specialization for bools/bitfields */
-	template<>
-	struct FGetValueFastHelper<bool>
-	{
-		static bool GetValue(const FCachedPropertyPath& InPropertyPath, bool& OutValue, FProperty*& OutProperty)
-		{
-			const FPropertyPathSegment& LastSegment = InPropertyPath.GetLastSegment();
-			OutProperty = CastFieldChecked<FProperty>(LastSegment.GetField().ToField());
-			FArrayProperty* ArrayProp = CastField<FArrayProperty>(OutProperty);
-			if ( ArrayProp && LastSegment.GetArrayIndex() != INDEX_NONE )
-			{
-				if (IsConcreteTypeCompatibleWithReflectedType<bool>(ArrayProp->Inner))
-				{
-					ArrayProp->Inner->CopySingleValue(&OutValue, InPropertyPath.GetCachedAddress());
-					return true;
-				}
-			}
-			else if (IsConcreteTypeCompatibleWithReflectedType<bool>(OutProperty))
-			{
-				FBoolProperty* BoolProperty = CastFieldChecked<FBoolProperty>(OutProperty);
-				OutValue = BoolProperty->GetPropertyValue(InPropertyPath.GetCachedAddress());
-				return true;
-			}
-			return false;
-		}
-	};
-
 	/** 
 	 * Fast, unsafe version of GetValue().
 	 * @param	InContainer		The containing object/structure to iterate against
@@ -1191,17 +1307,57 @@ namespace PropertyPathHelpersInternal
 			const FPropertyPathSegment& LastSegment = InPropertyPath.GetLastSegment();
 			FProperty* Property = CastFieldChecked<FProperty>(LastSegment.GetField().ToField());
 			FArrayProperty* ArrayProp = CastField<FArrayProperty>(Property);
+
+			const void* Value = nullptr;
+			if constexpr (std::is_same<FPropertyStructView, T>::value)
+			{
+				const FPropertyStructView& InStuctView = static_cast<const FPropertyStructView&>(InValue);
+				Value = static_cast<const void*>(InStuctView.Memory);
+			}
+			else
+			{
+				Value = &InValue;
+			}
+
+			auto IsPropertyCompatible = [](FProperty* InProperty)
+			{
+				if constexpr (!std::is_same<FPropertyStructView, T>::value)
+				{
+					if (!IsConcreteTypeCompatibleWithReflectedType<T>(InProperty))
+					{
+						return false;
+					}
+				}
+
+				return true;
+			};
+
 			if ( ArrayProp && LastSegment.GetArrayIndex() != INDEX_NONE )
 			{
-				if (IsConcreteTypeCompatibleWithReflectedType<T>(ArrayProp->Inner))
+				if (!IsPropertyCompatible(ArrayProp->Inner))
 				{
-					ArrayProp->Inner->CopySingleValue(InPropertyPath.GetCachedAddress(), &InValue);
-					return true;
+					return false;
 				}
+
+				ArrayProp->Inner->CopySingleValue(InPropertyPath.GetCachedAddress(), Value);
+				return true;
 			}
-			else if (IsConcreteTypeCompatibleWithReflectedType<T>(Property))
+			else 
 			{
-				Property->CopySingleValue(InPropertyPath.GetCachedAddress(), &InValue);
+				if (!IsPropertyCompatible(Property))
+				{
+					return false;
+				}
+
+				if (Property->HasSetter())
+				{
+					Property->CallSetter(InPropertyPath.GetCachedContainer(), Value);
+				}
+				else
+				{
+					Property->CopySingleValue(InPropertyPath.GetCachedAddress(), Value);
+					CallParentSetters(InPropertyPath);
+				}
 				return true;
 			}
 			return false;
@@ -1244,11 +1400,19 @@ namespace PropertyPathHelpersInternal
 			}
 			else if (IsConcreteTypeCompatibleWithReflectedType<bool>(Property))
 			{
-				FBoolProperty* BoolProperty = CastFieldChecked<FBoolProperty>(Property);
-				BoolProperty->SetPropertyValue(InPropertyPath.GetCachedAddress(), InValue);
+				if (Property->HasSetter())
+				{
+					// Setter should be specialized to handle masking
+					Property->CallSetter(InPropertyPath.GetCachedContainer(), &InValue);
+				}
+				else
+				{
+					FBoolProperty* BoolProperty = CastFieldChecked<FBoolProperty>(Property);
+					BoolProperty->SetPropertyValue(InPropertyPath.GetCachedAddress(), InValue);
+					CallParentSetters(InPropertyPath);
+				}
 				return true;
 			}
-			
 			return false;
 		}
 	};
@@ -1274,4 +1438,4 @@ namespace PropertyPathHelpersInternal
 
 		return false;
 	}
-}
\ No newline at end of file
+}
diff --git a/Engine/Source/Runtime/PropertyPath/Public/PropertyTypeCompatibility.h b/Engine/Source/Runtime/PropertyPath/Public/PropertyTypeCompatibility.h
--- a/Engine/Source/Runtime/PropertyPath/Public/PropertyTypeCompatibility.h
+++ b/Engine/Source/Runtime/PropertyPath/Public/PropertyTypeCompatibility.h
@@ -383,24 +383,28 @@ struct FConcreteTypeCompatibleWithReflectedTypeHelper<TLazyObjectPtr<T>>
 /** 
  * Check whether the concrete type T is compatible with the reflected type of a FProperty for 
  * the purposes of CopysingleValue()
- * Non-enum implementation.
  */
 template<typename T>
-static inline typename TEnableIf<!TIsEnum<T>::Value, bool>::Type IsConcreteTypeCompatibleWithReflectedType(FProperty* Property)
+static inline bool IsConcreteTypeCompatibleWithReflectedType(FProperty* Property)
 {
-	return FConcreteTypeCompatibleWithReflectedTypeHelper<T>::IsConcreteTypeCompatibleWithReflectedType(Property);
-}
-
-/** Enum implementation. @see IsConcreteTypeCompatibleWithReflectedType */
-template<typename T>
-static inline typename TEnableIf<TIsEnum<T>::Value, bool>::Type IsConcreteTypeCompatibleWithReflectedType(FProperty* Property)
-{
-	return FConcreteTypeCompatibleWithReflectedTypeHelper<uint8>::IsConcreteTypeCompatibleWithReflectedType(Property);
+	if constexpr (TIsEnum<T>::Value || TIsTEnumAsByte<T>::Value)
+	{
+		return FConcreteTypeCompatibleWithReflectedTypeHelper<uint8>::IsConcreteTypeCompatibleWithReflectedType(Property);
+	}
+	else
+	{
+		return FConcreteTypeCompatibleWithReflectedTypeHelper<T>::IsConcreteTypeCompatibleWithReflectedType(Property);
+	}
 }
 
 template<typename T>
 inline bool PropertySizesMatch_Impl(FProperty* InProperty)
 {
+	if (FByteProperty* ByteProperty = CastField<FByteProperty>(InProperty))
+	{
+		return InProperty->ElementSize == sizeof(uint8);
+	}
+
 	return InProperty->ElementSize == sizeof(T);
 }
 
