Migrate Alloca calls to outside of loop to fix static analysis warnings

diff --git a/Engine/Source/Runtime/PropertyPath/Private/PropertyPathHelpers.cpp b/Engine/Source/Runtime/PropertyPath/Private/PropertyPathHelpers.cpp
--- a/Engine/Source/Runtime/PropertyPath/Private/PropertyPathHelpers.cpp
+++ b/Engine/Source/Runtime/PropertyPath/Private/PropertyPathHelpers.cpp
@@ -529,6 +529,9 @@ namespace PropertyPathHelpersInternal
 			? InPropertyPath.GetCachedContainer()
 			: InPropertyPath.GetCachedLastContainerInPath();
 
+		// Compilers will warn if Alloca is used inside a loop, so break Alloca logic out of the loop
+		int32 ParentSegmentIndex = INDEX_NONE;
+
 		// Call the topmost setter on the last UObject in path
 		const int32 NumSegments = InPropertyPath.GetNumSegments();
 		for (int32 Index = IndexAfterCachedLastContainer; Index < NumSegments; Index++)
@@ -540,33 +543,45 @@ namespace PropertyPathHelpersInternal
 
 			if (ParentProperty->HasSetter() && ParentArrayIndex < ParentProperty->ArrayDim)
 			{
-				// We want to call the setter with the current value, so just get the pointer to current value via container
-				if (void* ParentAddress = ParentProperty->ContainerPtrToValuePtr<void>(ContainerPtr, ParentArrayIndex))
+				ParentSegmentIndex = Index;
+				break;
+			}
+		}
+
+		// Note: We check for HasSetter & ParentArrayIndex valid above
+		if (ParentSegmentIndex != INDEX_NONE)
+		{
+			const FPropertyPathSegment& ParentSegment = InPropertyPath.GetSegment(ParentSegmentIndex);
+			int32 ParentArrayIndex = ParentSegment.GetArrayIndex();
+			FProperty* ParentProperty = CastFieldChecked<FProperty>(ParentSegment.GetField().ToField());
+			ParentArrayIndex = ParentArrayIndex == INDEX_NONE ? 0 : ParentArrayIndex;
+
+			// We want to call the setter with the current value, so just get the pointer to current value via container
+			if (void* ParentAddress = ParentProperty->ContainerPtrToValuePtr<void>(ContainerPtr, ParentArrayIndex))
+			{
+				if (ParentProperty->HasGetter())
 				{
-					if (ParentProperty->HasGetter())
-					{
-						// Call getter if it has one, getter SHOULD be pure and thus won't cause behavioral changes.
-						// But this is a read on ParentAddress so we call getter for now
+					// Call getter if it has one, getter SHOULD be pure and thus won't cause behavioral changes.
+					// But this is a read on ParentAddress so we call getter for now
 
-						int32 Size = ParentSegment.GetStruct()->GetPropertiesSize();
-						int32 Alignment = ParentSegment.GetStruct()->GetMinAlignment();
-						uint8* Temp = (uint8*)FMemory_Alloca_Aligned(Size, Alignment);
-						FMemory::Memzero(Temp, Size);
+					int32 Size = ParentSegment.GetStruct()->GetPropertiesSize();
+					int32 Alignment = ParentSegment.GetStruct()->GetMinAlignment();
+					uint8* Temp = (uint8*)FMemory_Alloca_Aligned(Size, Alignment);
+					FMemory::Memzero(Temp, Size);
 
-						if (!ParentProperty->HasAnyPropertyFlags(CPF_ZeroConstructor))
-						{
-							ParentProperty->InitializeValue_InContainer(Temp);
-						}
-						ParentProperty->CallGetter(ContainerPtr, Temp);
-						ParentProperty->CallSetter(ContainerPtr, Temp);
-						ParentProperty->DestroyValue_InContainer(Temp);
-					}
-					else
+					if (!ParentProperty->HasAnyPropertyFlags(CPF_ZeroConstructor))
 					{
-						ParentProperty->CallSetter(ContainerPtr, ParentAddress);
+						ParentProperty->InitializeValue_InContainer(Temp);
 					}
-					return;
+					ParentProperty->CallGetter(ContainerPtr, Temp);
+					ParentProperty->CallSetter(ContainerPtr, Temp);
+					ParentProperty->DestroyValue_InContainer(Temp);
+				}
+				else
+				{
+					ParentProperty->CallSetter(ContainerPtr, ParentAddress);
 				}
+				return;
 			}
 		}
 	}
@@ -594,6 +609,9 @@ namespace PropertyPathHelpersInternal
 			}
 		};
 
+		// Compilers will warn if Alloca is used inside a loop, so break Alloca logic out of the loop
+		int32 ParentSegmentIndex = INDEX_NONE;
+
 		// Call the topmost getter on the last UObject in path
 		const int32 NumSegments = InPropertyPath.GetNumSegments();
 		for (int32 Index = IndexAfterCachedLastContainer; Index < NumSegments; Index++)
@@ -605,29 +623,41 @@ namespace PropertyPathHelpersInternal
 
 			if (ParentProperty->HasGetter() && ParentArrayIndex < ParentProperty->ArrayDim)
 			{
-				// We want to call the Getter with the current value, so just get the pointer to current value via container
-				void* ParentAddress = ParentProperty->ContainerPtrToValuePtr<void>(ContainerPtr, ParentArrayIndex);
-				if (ensure(ParentAddress))
-				{
-					int32 Size = ParentSegment.GetStruct()->GetPropertiesSize();
-					int32 Alignment = ParentSegment.GetStruct()->GetMinAlignment();
-					uint8* Temp = (uint8*)FMemory_Alloca_Aligned(Size, Alignment);
-					FMemory::Memzero(Temp, Size);
+				ParentSegmentIndex = Index;
+				break;
+			}
+		}
 
-					if (!ParentProperty->HasAnyPropertyFlags(CPF_ZeroConstructor))
-					{
-						ParentProperty->InitializeValue_InContainer(Temp);
-					}
-					ParentProperty->CallGetter(ContainerPtr, Temp);
+		// Note: We check for HasGetter & ParentArrayIndex valid above
+		if (ParentSegmentIndex != INDEX_NONE)
+		{
+			const FPropertyPathSegment& ParentSegment = InPropertyPath.GetSegment(ParentSegmentIndex);
+			int32 ParentArrayIndex = ParentSegment.GetArrayIndex();
+			FProperty* ParentProperty = CastFieldChecked<FProperty>(ParentSegment.GetField().ToField());
+			ParentArrayIndex = ParentArrayIndex == INDEX_NONE ? 0 : ParentArrayIndex;
 
-					// We resolved the property address earlier & it's containing UObject, use this for relative-offset for Temp
-					uint8* TempPropertyAddress = Temp + ((uint8*)InPropertyAddress - (uint8*)ContainerPtr);
-					GetValueFromProperty(OutValue, InPropertyPath, TempPropertyAddress);
-					ParentProperty->DestroyValue_InContainer(Temp);
+			// We want to call the Getter with the current value, so just get the pointer to current value via container
+			void* ParentAddress = ParentProperty->ContainerPtrToValuePtr<void>(ContainerPtr, ParentArrayIndex);
+			if (ensure(ParentAddress))
+			{
+				int32 Size = ParentSegment.GetStruct()->GetPropertiesSize();
+				int32 Alignment = ParentSegment.GetStruct()->GetMinAlignment();
+				uint8* Temp = (uint8*)FMemory_Alloca_Aligned(Size, Alignment);
+				FMemory::Memzero(Temp, Size);
+
+				if (!ParentProperty->HasAnyPropertyFlags(CPF_ZeroConstructor))
+				{
+					ParentProperty->InitializeValue_InContainer(Temp);
 				}
+				ParentProperty->CallGetter(ContainerPtr, Temp);
 
-				return;
+				// We resolved the property address earlier & it's containing UObject, use this for relative-offset for Temp
+				uint8* TempPropertyAddress = Temp + ((uint8*)InPropertyAddress - (uint8*)ContainerPtr);
+				GetValueFromProperty(OutValue, InPropertyPath, TempPropertyAddress);
+				ParentProperty->DestroyValue_InContainer(Temp);
 			}
+
+			return;
 		}
 
 		GetValueFromProperty(OutValue, InPropertyPath, InPropertyAddress);
