Fix propertypath getters resolving against UObject container instead of parent property container.

diff --git a/Engine/Source/Runtime/PropertyPath/Private/PropertyPathHelpers.cpp b/Engine/Source/Runtime/PropertyPath/Private/PropertyPathHelpers.cpp
--- a/Engine/Source/Runtime/PropertyPath/Private/PropertyPathHelpers.cpp
+++ b/Engine/Source/Runtime/PropertyPath/Private/PropertyPathHelpers.cpp
@@ -525,7 +525,7 @@ namespace PropertyPathHelpersInternal
 		int32 LastContainerInPathIndex = InPropertyPath.GetCachedLastContainerInPathIndex();
 		int32 IndexAfterCachedLastContainer = InPropertyPath.GetCachedLastContainerInPathIndex() + 1;
 
-		void* ContainerPtr = LastContainerInPathIndex == INDEX_NONE 
+		void* ObjectContainerPtr = LastContainerInPathIndex == INDEX_NONE 
 			? InPropertyPath.GetCachedContainer()
 			: InPropertyPath.GetCachedLastContainerInPath();
 
@@ -557,7 +557,7 @@ namespace PropertyPathHelpersInternal
 			ParentArrayIndex = ParentArrayIndex == INDEX_NONE ? 0 : ParentArrayIndex;
 
 			// We want to call the setter with the current value, so just get the pointer to current value via container
-			if (void* ParentAddress = ParentProperty->ContainerPtrToValuePtr<void>(ContainerPtr, ParentArrayIndex))
+			if (void* ParentAddress = ParentProperty->ContainerPtrToValuePtr<void>(ObjectContainerPtr, ParentArrayIndex))
 			{
 				if (ParentProperty->HasGetter())
 				{
@@ -573,13 +573,13 @@ namespace PropertyPathHelpersInternal
 					{
 						ParentProperty->InitializeValue_InContainer(Temp);
 					}
-					ParentProperty->CallGetter(ContainerPtr, Temp);
-					ParentProperty->CallSetter(ContainerPtr, Temp);
+					ParentProperty->CallGetter(ObjectContainerPtr, Temp);
+					ParentProperty->CallSetter(ObjectContainerPtr, Temp);
 					ParentProperty->DestroyValue_InContainer(Temp);
 				}
 				else
 				{
-					ParentProperty->CallSetter(ContainerPtr, ParentAddress);
+					ParentProperty->CallSetter(ObjectContainerPtr, ParentAddress);
 				}
 				return;
 			}
@@ -591,7 +591,7 @@ namespace PropertyPathHelpersInternal
 		int32 LastContainerInPathIndex = InPropertyPath.GetCachedLastContainerInPathIndex();
 		int32 IndexAfterCachedLastContainer = InPropertyPath.GetCachedLastContainerInPathIndex() + 1;
 
-		void* ContainerPtr = LastContainerInPathIndex == INDEX_NONE 
+		void* ObjectContainerPtr = LastContainerInPathIndex == INDEX_NONE 
 			? InPropertyPath.GetCachedContainer()
 			: InPropertyPath.GetCachedLastContainerInPath();
 
@@ -637,7 +637,7 @@ namespace PropertyPathHelpersInternal
 			ParentArrayIndex = ParentArrayIndex == INDEX_NONE ? 0 : ParentArrayIndex;
 
 			// We want to call the Getter with the current value, so just get the pointer to current value via container
-			void* ParentAddress = ParentProperty->ContainerPtrToValuePtr<void>(ContainerPtr, ParentArrayIndex);
+			void* ParentAddress = ParentProperty->ContainerPtrToValuePtr<void>(ObjectContainerPtr, ParentArrayIndex);
 			if (ensure(ParentAddress))
 			{
 				int32 Size = ParentSegment.GetStruct()->GetPropertiesSize();
@@ -649,10 +649,10 @@ namespace PropertyPathHelpersInternal
 				{
 					ParentProperty->InitializeValue_InContainer(Temp);
 				}
-				ParentProperty->CallGetter(ContainerPtr, Temp);
+				ParentProperty->CallGetter(ObjectContainerPtr, Temp);
 
-				// We resolved the property address earlier & it's containing UObject, use this for relative-offset for Temp
-				uint8* TempPropertyAddress = Temp + ((uint8*)InPropertyAddress - (uint8*)ContainerPtr);
+				// We resolved the property address earlier & it's containing parent, use this for relative-offset for Temp
+				uint8* TempPropertyAddress = Temp + ((uint8*)InPropertyAddress - (uint8*)ParentAddress);
 				GetValueFromProperty(OutValue, InPropertyPath, TempPropertyAddress);
 				ParentProperty->DestroyValue_InContainer(Temp);
 			}
diff --git a/Engine/Source/Runtime/PropertyPath/Private/Tests/PropertyPathHelpersTest.cpp b/Engine/Source/Runtime/PropertyPath/Private/Tests/PropertyPathHelpersTest.cpp
--- a/Engine/Source/Runtime/PropertyPath/Private/Tests/PropertyPathHelpersTest.cpp
+++ b/Engine/Source/Runtime/PropertyPath/Private/Tests/PropertyPathHelpersTest.cpp
@@ -16,8 +16,10 @@ DEFINE_LOG_CATEGORY_STATIC(LogPropertyPathHelpersTest, Log, All);
  * SetString
  * SetEnum
  * SetStruct
+ * SetInnerObject
  * SetStructData
  * SetStructMember
+ * SetStructInnerObjectMember
  * SetInnerStructMember
  * SetPrimitivesInnerObject
  * SetStringInnerObject
@@ -29,8 +31,10 @@ DEFINE_LOG_CATEGORY_STATIC(LogPropertyPathHelpersTest, Log, All);
  * GetPrimitives
  * GetString
  * GetStruct
+ * GetInnerObject
  * GetEnum
  * GetStructMember
+ * GetStructInnerObjectMember
  * GetInnerStructMember
  * GetPrimitivesInnerObject
  * GetStringInnerObject
@@ -288,6 +292,27 @@ bool FSetStructTest::RunTest(const FString& Parameters)
 	return GPassing;
 }
 
+IMPLEMENT_SIMPLE_AUTOMATION_TEST(FSetInnerObjectTest, "System.PropertyPath.SetInnerObjectTest", EAutomationTestFlags::ApplicationContextMask | EAutomationTestFlags::EngineFilter)
+bool FSetInnerObjectTest::RunTest(const FString& Parameters)
+{
+	bool GPassing = true;
+
+	{
+		FPropertyPathTestBed Test = {};
+		FPropertyPathTestBed CppTest = {};
+		CppTest.Object->InnerObject = CppTest.ModifiedObject;
+
+		PropertyPathHelpers::SetPropertyValue(Test.Object, FString("InnerObject"), Test.ModifiedObject);
+		GPassing &= TestEqual(TEXT("Inner Object Set"), Test.Object->InnerObject, Test.ModifiedObject);
+		GPassing &= TestEqual(TEXT("Inner Object Set Cpp Equivalent"), *Test.Object, *CppTest.Object);
+		GPassing &= TestEqual(TEXT("Inner Object Setter Called"), Test.Object->IsSetterCalled(), false);
+		GPassing &= TestEqual(TEXT("Inner Object Getter Called"), Test.Object->IsGetterCalled(), false);
+		Test.Object->ResetGetterSetterFlags();
+	}
+
+	return GPassing;
+}
+
 IMPLEMENT_SIMPLE_AUTOMATION_TEST(FSetStructDataTest, "System.PropertyPath.SetStructDataTest", EAutomationTestFlags::ApplicationContextMask | EAutomationTestFlags::EngineFilter)
 bool FSetStructDataTest::RunTest(const FString& Parameters)
 {
@@ -397,6 +422,65 @@ bool FSetStructMemberTest::RunTest(const FString& Parameters)
 	return GPassing;
 }
 
+IMPLEMENT_SIMPLE_AUTOMATION_TEST(FSetStructInnerObjectMemberTest, "System.PropertyPath.SetStructInnerObjectMemberTest", EAutomationTestFlags::ApplicationContextMask | EAutomationTestFlags::EngineFilter)
+bool FSetStructInnerObjectMemberTest::RunTest(const FString& Parameters)
+{
+	bool GPassing = true;
+
+	{
+		FPropertyPathTestBed Test = {};
+		FPropertyPathTestBed CppTest = {};
+		CppTest.Object->Struct.InnerObject = CppTest.ModifiedObject;
+		CppTest.Object->SetStruct(CppTest.Object->Struct);
+
+		PropertyPathHelpers::SetPropertyValue(Test.Object, FString("Struct.InnerObject"), Test.ModifiedObject);
+		GPassing &= TestEqual(TEXT("Struct Inner Object Member Set"), Test.Object->Struct.InnerObject, Test.ModifiedObject);
+		GPassing &= TestEqual(TEXT("Struct Inner Object Member Set Cpp Equivalent"), *Test.Object, *CppTest.Object);
+		GPassing &= TestEqual(TEXT("Struct Setter Called"), Test.Object->IsSetterCalled(), true);
+
+		// Emphasis, see note under coverage at top
+		GPassing &= TestEqual(TEXT("Struct Getter Called"), Test.Object->IsGetterCalled(), true);
+
+		Test.Object->ResetGetterSetterFlags();
+	}
+
+	{
+		FPropertyPathTestBed Test = {};
+		FPropertyPathTestBed CppTest = {};
+		CppTest.Object->StructRef.InnerObject = CppTest.ModifiedObject;
+		CppTest.Object->SetStructRef(CppTest.Object->StructRef);
+
+		PropertyPathHelpers::SetPropertyValue(Test.Object, FString("StructRef.InnerObject"), Test.ModifiedObject);
+		GPassing &= TestEqual(TEXT("Struct Ref Inner Object Member Set"), Test.Object->StructRef.InnerObject, Test.ModifiedObject);
+		GPassing &= TestEqual(TEXT("Struct Ref Inner Object Member Set Cpp Equivalent"), *Test.Object, *CppTest.Object);
+		GPassing &= TestEqual(TEXT("Struct Setter Called"), Test.Object->IsSetterCalled(), true);
+
+		// Emphasis, see note under coverage at top
+		GPassing &= TestEqual(TEXT("Struct Getter Called"), Test.Object->IsGetterCalled(), true);
+
+		Test.Object->ResetGetterSetterFlags();
+	}
+
+	{
+		FPropertyPathTestBed Test = {};
+		FPropertyPathTestBed CppTest = {};
+		CppTest.Object->StructConstRef.InnerObject = CppTest.ModifiedObject;
+		CppTest.Object->SetStructConstRef(CppTest.Object->StructConstRef);
+
+		PropertyPathHelpers::SetPropertyValue(Test.Object, FString("StructConstRef.InnerObject"), Test.ModifiedObject);
+		GPassing &= TestEqual(TEXT("Struct Const Ref Inner Object Member Set"), Test.Object->StructConstRef.InnerObject, Test.ModifiedObject);
+		GPassing &= TestEqual(TEXT("Struct Const Ref Inner Object Member Set Cpp Equivalent"), *Test.Object, *CppTest.Object);
+		GPassing &= TestEqual(TEXT("Struct Setter Called"), Test.Object->IsSetterCalled(), true);
+
+		// Emphasis, see note under coverage at top
+		GPassing &= TestEqual(TEXT("Struct Getter Called"), Test.Object->IsGetterCalled(), true);
+
+		Test.Object->ResetGetterSetterFlags();
+	}
+
+	return GPassing;
+}
+
 IMPLEMENT_SIMPLE_AUTOMATION_TEST(FSetInnerStructMemberTest, "System.PropertyPath.SetInnerStructMemberTest", EAutomationTestFlags::ApplicationContextMask | EAutomationTestFlags::EngineFilter)
 bool FSetInnerStructMemberTest::RunTest(const FString& Parameters)
 {
@@ -743,6 +827,27 @@ bool FGetStructTest::RunTest(const FString& Parameters)
 	return GPassing;
 }
 
+IMPLEMENT_SIMPLE_AUTOMATION_TEST(FGetInnerObjectTest, "System.PropertyPath.GetInnerObjectTest", EAutomationTestFlags::ApplicationContextMask | EAutomationTestFlags::EngineFilter)
+bool FGetInnerObjectTest::RunTest(const FString& Parameters)
+{
+	bool GPassing = true;
+
+	{
+		FPropertyPathTestBed Test = {};
+		FPropertyPathTestBed CppTest = {};
+		CppTest.ModifiedObject = CppTest.Object->InnerObject;
+
+		PropertyPathHelpers::GetPropertyValue(Test.Object, FString("InnerObject"), Test.ModifiedObject);
+		GPassing &= TestEqual(TEXT("Inner Object Get"), Test.Object->InnerObject, Test.ModifiedObject);
+		GPassing &= TestEqual(TEXT("Inner Object Get Cpp Equivalent"), *CppTest.ModifiedObject, *Test.ModifiedObject);
+		GPassing &= TestEqual(TEXT("Inner Object Setter Called"), Test.Object->IsSetterCalled(), false);
+		GPassing &= TestEqual(TEXT("Inner Object Getter Called"), Test.Object->IsGetterCalled(), false);
+		Test.Object->ResetGetterSetterFlags();
+	}
+
+	return GPassing;
+}
+
 IMPLEMENT_SIMPLE_AUTOMATION_TEST(FGetStructMemberTest, "System.PropertyPath.GetStructMemberTest", EAutomationTestFlags::ApplicationContextMask | EAutomationTestFlags::EngineFilter)
 bool FGetStructMemberTest::RunTest(const FString& Parameters)
 {
@@ -793,6 +898,56 @@ bool FGetStructMemberTest::RunTest(const FString& Parameters)
 	return GPassing;
 }
 
+IMPLEMENT_SIMPLE_AUTOMATION_TEST(FGetStructInnerObjectMemberTest, "System.PropertyPath.GetStructInnerObjectMemberTest", EAutomationTestFlags::ApplicationContextMask | EAutomationTestFlags::EngineFilter)
+bool FGetStructInnerObjectMemberTest::RunTest(const FString& Parameters)
+{
+	bool GPassing = true;
+
+	{
+		FPropertyPathTestBed Test = {};
+		FPropertyPathTestBed CppTest = {};
+		CppTest.ModifiedObject = CppTest.Object->GetStruct().InnerObject;
+
+		PropertyPathHelpers::GetPropertyValue(Test.Object, FString("Struct.InnerObject"), Test.ModifiedObject);
+		GPassing &= TestEqual(TEXT("Struct Member Get"), *Test.ModifiedObject, *Test.Object->Struct.InnerObject);
+		GPassing &= TestEqual(TEXT("Struct Member Get Cpp Equivalent"), *CppTest.ModifiedObject, *Test.ModifiedObject);
+		GPassing &= TestEqual(TEXT("Struct Setter Called"), Test.Object->IsSetterCalled(), false);
+		GPassing &= TestEqual(TEXT("Struct Getter Called"), Test.Object->IsGetterCalled(), true);
+
+		Test.Object->ResetGetterSetterFlags();
+	}
+
+	{
+		FPropertyPathTestBed Test = {};
+		FPropertyPathTestBed CppTest = {};
+		CppTest.ModifiedObject = CppTest.Object->GetStructRef().InnerObject;
+
+		PropertyPathHelpers::GetPropertyValue(Test.Object, FString("StructRef.InnerObject"), Test.ModifiedObject);
+		GPassing &= TestEqual(TEXT("Struct Ref Member Get"), *Test.ModifiedObject, *Test.Object->StructRef.InnerObject);
+		GPassing &= TestEqual(TEXT("Struct Ref Member Get Cpp Equivalent"), *CppTest.ModifiedObject, *Test.ModifiedObject);
+		GPassing &= TestEqual(TEXT("Struct Setter Called"), Test.Object->IsSetterCalled(), false);
+		GPassing &= TestEqual(TEXT("Struct Getter Called"), Test.Object->IsGetterCalled(), true);
+
+		Test.Object->ResetGetterSetterFlags();
+	}
+
+	{
+		FPropertyPathTestBed Test = {};
+		FPropertyPathTestBed CppTest = {};
+		CppTest.ModifiedObject = CppTest.Object->GetStructConstRef().InnerObject;
+
+		PropertyPathHelpers::GetPropertyValue(Test.Object, FString("StructConstRef.InnerObject"), Test.ModifiedObject);
+		GPassing &= TestEqual(TEXT("Struct Const Ref Member Get"), *Test.ModifiedObject, *Test.Object->StructConstRef.InnerObject);
+		GPassing &= TestEqual(TEXT("Struct Const Ref Member Get Cpp Equivalent"), *CppTest.ModifiedObject, *Test.ModifiedObject);
+		GPassing &= TestEqual(TEXT("Struct Setter Called"), Test.Object->IsSetterCalled(), false);
+		GPassing &= TestEqual(TEXT("Struct Getter Called"), Test.Object->IsGetterCalled(), true);
+
+		Test.Object->ResetGetterSetterFlags();
+	}
+
+	return GPassing;
+}
+
 IMPLEMENT_SIMPLE_AUTOMATION_TEST(FGetInnerStructMemberTest, "System.PropertyPath.GetInnerStructMemberTest", EAutomationTestFlags::ApplicationContextMask | EAutomationTestFlags::EngineFilter)
 bool FGetInnerStructMemberTest::RunTest(const FString& Parameters)
 {
diff --git a/Engine/Source/Runtime/PropertyPath/Private/Tests/PropertyPathHelpersTest.h b/Engine/Source/Runtime/PropertyPath/Private/Tests/PropertyPathHelpersTest.h
--- a/Engine/Source/Runtime/PropertyPath/Private/Tests/PropertyPathHelpersTest.h
+++ b/Engine/Source/Runtime/PropertyPath/Private/Tests/PropertyPathHelpersTest.h
@@ -8,6 +8,8 @@
 #include "PropertyPathHelpers.h"
 #include "PropertyPathHelpersTest.generated.h"
 
+class UPropertyPathTestObject;
+
 UENUM()
 enum EPropertyPathTestEnum : int
 {
@@ -17,7 +19,6 @@ enum EPropertyPathTestEnum : int
 	Four
 };
 
-
 USTRUCT(BlueprintType)
 struct FPropertyPathTestInnerStruct
 {
@@ -121,6 +122,9 @@ struct FPropertyPathTestStruct
 	UPROPERTY()
 	FPropertyPathTestInnerStruct InnerStruct = {};
 
+	UPROPERTY()
+	TObjectPtr<UPropertyPathTestObject> InnerObject = nullptr;
+
 	bool operator ==(const FPropertyPathTestStruct& Other) const
 	{
 		return Bool == Other.Bool
@@ -309,6 +313,9 @@ struct FPropertyPathTestBed
 	{
 		Object = NewObject<UPropertyPathTestObject>();
 		Object->InnerObject = NewObject<UPropertyPathTestObject>();
+		Object->Struct.InnerObject = NewObject<UPropertyPathTestObject>();
+		Object->StructRef.InnerObject = NewObject<UPropertyPathTestObject>();
+		Object->StructConstRef.InnerObject = NewObject<UPropertyPathTestObject>();
 
 		ModifiedStruct = {};
 
@@ -331,11 +338,25 @@ struct FPropertyPathTestBed
 		ModifiedStruct.InnerStruct.Float = 1.5f;
 
 		DefaultStruct = {};
+
+		ModifiedObject = NewObject<UPropertyPathTestObject>();
+
+		ModifiedObject->Bool = true;
+		ModifiedObject->Integer = 1;
+		ModifiedObject->EnumOne = Two;
+		ModifiedObject->EnumTwo = Two;
+		ModifiedObject->EnumThree = Two;
+		ModifiedObject->EnumFour = Two;
+		ModifiedObject->String = "NewValue";
+		ModifiedObject->Float = 1.5f;
 	}
 
 	UPROPERTY()
 	TObjectPtr<UPropertyPathTestObject> Object;
 
+	UPROPERTY()
+	TObjectPtr<UPropertyPathTestObject> ModifiedObject;
+
 	UPROPERTY()
 	FPropertyPathTestStruct ModifiedStruct;
 
diff --git a/Engine/Source/Runtime/PropertyPath/Public/PropertyTypeCompatibility.h b/Engine/Source/Runtime/PropertyPath/Public/PropertyTypeCompatibility.h
--- a/Engine/Source/Runtime/PropertyPath/Public/PropertyTypeCompatibility.h
+++ b/Engine/Source/Runtime/PropertyPath/Public/PropertyTypeCompatibility.h
@@ -313,25 +313,26 @@ inline bool IsConcreteTypeCompatibleWithReflectedType_Impl<FSoftClassPath>(FProp
 	return IsConcreteTypeCompatibleWithReflectedType_BuiltInStruct<FSoftClassPath>(Property);
 }
 
-template<>
-inline bool IsConcreteTypeCompatibleWithReflectedType_Impl<UObject*>(FProperty* Property)
-{
-	if ( FObjectProperty* ObjectProperty = CastField<FObjectProperty>(Property) )
-	{
-		return true;
-		//return ObjectProperty->PropertyClass->IsChildOf(T::StaticClass());
-	}
-
-	return false;
-}
-
 /** Standard implementation */
 template<typename T> 
 struct FConcreteTypeCompatibleWithReflectedTypeHelper
 {
 	static bool IsConcreteTypeCompatibleWithReflectedType(FProperty* Property) 
 	{
-		return IsConcreteTypeCompatibleWithReflectedType_Impl<T>(Property);
+		if constexpr (TIsPointerOrObjectPtrToBaseOf<T, UObject>::Value)
+		{
+			if (FObjectProperty* ObjectProperty = CastField<FObjectProperty>(Property))
+			{
+				return true;
+				//return ObjectProperty->PropertyClass->IsChildOf(T::StaticClass());
+			}
+
+			return false;
+		}
+		else
+		{
+			return IsConcreteTypeCompatibleWithReflectedType_Impl<T>(Property);
+		}
 	}
 };
 
