Possible fix for unresolved external destructor of ~FCachedPropertyPath() under some build environments.

diff --git a/Engine/Source/Runtime/PropertyPath/Private/PropertyPathHelpers.cpp b/Engine/Source/Runtime/PropertyPath/Private/PropertyPathHelpers.cpp
--- a/Engine/Source/Runtime/PropertyPath/Private/PropertyPathHelpers.cpp
+++ b/Engine/Source/Runtime/PropertyPath/Private/PropertyPathHelpers.cpp
@@ -789,6 +789,8 @@ FCachedPropertyPath::FCachedPropertyPath(const FPropertyPathSegment& Segment)
 	Segments.Add(Segment);
 }
 
+FCachedPropertyPath::~FCachedPropertyPath() = default;
+
 void FCachedPropertyPath::MakeFromString(const FString& InPropertyPath)
 {
 	const TCHAR Delim = TEXT('.');
diff --git a/Engine/Source/Runtime/PropertyPath/Public/PropertyPathHelpers.h b/Engine/Source/Runtime/PropertyPath/Public/PropertyPathHelpers.h
--- a/Engine/Source/Runtime/PropertyPath/Public/PropertyPathHelpers.h
+++ b/Engine/Source/Runtime/PropertyPath/Public/PropertyPathHelpers.h
@@ -109,6 +109,9 @@ struct PROPERTYPATH_API FCachedPropertyPath
 	/** */
 	FCachedPropertyPath(const TArray<FString>& PathSegments);
 
+	/** */
+	~FCachedPropertyPath();
+
 	/** Check whether this property path is non-empty */
 	bool IsValid() const { return Segments.Num() > 0; }
 
