Fix missing vtable causing crashes when using property paths

#tests unit, PIE frontend

diff --git a/Engine/Source/Runtime/PropertyPath/Private/PropertyPathHelpers.cpp b/Engine/Source/Runtime/PropertyPath/Private/PropertyPathHelpers.cpp
--- a/Engine/Source/Runtime/PropertyPath/Private/PropertyPathHelpers.cpp
+++ b/Engine/Source/Runtime/PropertyPath/Private/PropertyPathHelpers.cpp
@@ -569,13 +569,10 @@ namespace PropertyPathHelpersInternal
 					uint8* Temp = (uint8*)FMemory_Alloca_Aligned(Size, Alignment);
 					FMemory::Memzero(Temp, Size);
 
-					if (!ParentProperty->HasAnyPropertyFlags(CPF_ZeroConstructor))
-					{
-						ParentProperty->InitializeValue_InContainer(Temp);
-					}
+					ParentProperty->InitializeValue(Temp);
 					ParentProperty->CallGetter(ObjectContainerPtr, Temp);
 					ParentProperty->CallSetter(ObjectContainerPtr, Temp);
-					ParentProperty->DestroyValue_InContainer(Temp);
+					ParentProperty->DestroyValue(Temp);
 				}
 				else
 				{
@@ -645,16 +642,13 @@ namespace PropertyPathHelpersInternal
 				uint8* Temp = (uint8*)FMemory_Alloca_Aligned(Size, Alignment);
 				FMemory::Memzero(Temp, Size);
 
-				if (!ParentProperty->HasAnyPropertyFlags(CPF_ZeroConstructor))
-				{
-					ParentProperty->InitializeValue_InContainer(Temp);
-				}
+				ParentProperty->InitializeValue(Temp);
 				ParentProperty->CallGetter(ObjectContainerPtr, Temp);
 
 				// We resolved the property address earlier & it's containing parent, use this for relative-offset for Temp
 				uint8* TempPropertyAddress = Temp + ((uint8*)InPropertyAddress - (uint8*)ParentAddress);
 				GetValueFromProperty(OutValue, InPropertyPath, TempPropertyAddress);
-				ParentProperty->DestroyValue_InContainer(Temp);
+				ParentProperty->DestroyValue(Temp);
 			}
 
 			return;
diff --git a/Engine/Source/Runtime/PropertyPath/Private/Tests/PropertyPathHelpersTest.h b/Engine/Source/Runtime/PropertyPath/Private/Tests/PropertyPathHelpersTest.h
--- a/Engine/Source/Runtime/PropertyPath/Private/Tests/PropertyPathHelpersTest.h
+++ b/Engine/Source/Runtime/PropertyPath/Private/Tests/PropertyPathHelpersTest.h
@@ -19,12 +19,21 @@ enum EPropertyPathTestEnum : int
 	Four
 };
 
+USTRUCT()
+struct FPropertyPathTestBaseStruct
+{
+	GENERATED_BODY()
+
+public:
+	virtual ~FPropertyPathTestBaseStruct() { }
+};
+
 USTRUCT(BlueprintType)
-struct FPropertyPathTestInnerStruct
+struct FPropertyPathTestInnerStruct : public FPropertyPathTestBaseStruct
 {
 	GENERATED_BODY()
 
-	~FPropertyPathTestInnerStruct()
+	virtual ~FPropertyPathTestInnerStruct()
 	{
 		Float = -1.5f;
 		Bool = true;
@@ -79,11 +88,11 @@ struct FPropertyPathTestInnerStruct
 };
 
 USTRUCT(BlueprintType)
-struct FPropertyPathTestStruct
+struct FPropertyPathTestStruct : public FPropertyPathTestBaseStruct
 {
 	GENERATED_BODY()
 
-	~FPropertyPathTestStruct()
+	virtual ~FPropertyPathTestStruct()
 	{
 		Float = -1.5f;
 		Bool = true;
