PR #11597: UMG Fixing CDO De-Sync With Root Widget Property Changes

diff --git a/Engine/Source/Editor/UMGEditor/Private/Details/SWidgetDetailsView.cpp b/Engine/Source/Editor/UMGEditor/Private/Details/SWidgetDetailsView.cpp
--- a/Engine/Source/Editor/UMGEditor/Private/Details/SWidgetDetailsView.cpp
+++ b/Engine/Source/Editor/UMGEditor/Private/Details/SWidgetDetailsView.cpp
@@ -624,7 +624,7 @@ void SWidgetDetailsView::NotifyPreChange(FEditPropertyChain* PropertyAboutToChan
 		TSharedPtr<FWidgetBlueprintEditor> Editor = BlueprintEditor.Pin();
 
 		const bool bIsModify = true;
-		Editor->MigrateFromChain(PropertyAboutToChange, bIsModify);
+		Editor->MigrateFromChain(nullptr, PropertyAboutToChange, bIsModify);
 	}
 }
 
@@ -637,7 +637,7 @@ void SWidgetDetailsView::NotifyPostChange(const FPropertyChangedEvent& PropertyC
 		TSharedPtr<FWidgetBlueprintEditor> Editor = BlueprintEditor.Pin();
 
 		const bool bIsModify = false;
-		Editor->MigrateFromChain(PropertyThatChanged, bIsModify);
+		Editor->MigrateFromChain(&PropertyChangedEvent, PropertyThatChanged, bIsModify);
 
 		// Any time we migrate a property value we need to mark the blueprint as structurally modified so users don't need 
 		// to recompile it manually before they see it play in game using the latest version.
@@ -658,4 +658,4 @@ void SWidgetDetailsView::NotifyPostChange(const FPropertyChangedEvent& PropertyC
 	}
 }
 
-#undef LOCTEXT_NAMESPACE
+#undef LOCTEXT_NAMESPACE
\ No newline at end of file
diff --git a/Engine/Source/Editor/UMGEditor/Private/WidgetBlueprintEditor.cpp b/Engine/Source/Editor/UMGEditor/Private/WidgetBlueprintEditor.cpp
--- a/Engine/Source/Editor/UMGEditor/Private/WidgetBlueprintEditor.cpp
+++ b/Engine/Source/Editor/UMGEditor/Private/WidgetBlueprintEditor.cpp
@@ -1073,19 +1073,164 @@ void FWidgetBlueprintEditor::AddReferencedObjects( FReferenceCollector& Collecto
 	Collector.AddReferencedObject(PreviewWidgetPtr);
 }
 
+static void PropagateDefaultPropertyChange(const FPropertyChangedEvent& PropertyChangedEvent, const FEditPropertyChain& PropertyThatChanged, UObject* UnmodifiedCDOObject, UObject* ModifedProxyObject)
+{
+	if (!UnmodifiedCDOObject || !ModifedProxyObject || !UnmodifiedCDOObject->HasAnyFlags(RF_ClassDefaultObject))
+	{
+		return;
+	}
+
+	// This is a kit bashed version of what happens in FPropertyNode::PropagatePropertyChange to migrate the CDO edits to the live instances.
+	// We don't have PropertyHandles outside of the details view, and so we need to do the equivalent work with property paths.
+	FCachedPropertyPath ChangedPropertyPath(PropertyChangedEvent, PropertyThatChanged);
+
+	FCachedPropertyPath ResolvedCDOPath = FCachedPropertyPath::MakeUnresolvedCopy(ChangedPropertyPath);
+	ResolvedCDOPath.Resolve(UnmodifiedCDOObject);
+
+	if (ResolvedCDOPath.IsFullyResolved())
+	{
+		FProperty* Prop = ResolvedCDOPath.GetFProperty();
+		FProperty* ParentProp = nullptr;
+		if (ResolvedCDOPath.GetNumSegments() > 1)
+		{
+			ParentProp = CastField<FProperty>(ResolvedCDOPath.GetSegment(ResolvedCDOPath.GetNumSegments() - 2).GetField().ToField());
+		}
+
+		FArrayProperty* ParentArrayProp = CastField<FArrayProperty>(ParentProp);
+		FMapProperty* ParentMapProp = CastField<FMapProperty>(ParentProp);
+		FSetProperty* ParentSetProp = CastField<FSetProperty>(ParentProp);
+
+		if (ParentArrayProp && ParentArrayProp->Inner != Prop)
+		{
+			ParentArrayProp = nullptr;
+		}
+
+		if (ParentMapProp && ParentMapProp->KeyProp != Prop && ParentMapProp->ValueProp != Prop)
+		{
+			ParentMapProp = nullptr;
+		}
+
+		if (ParentSetProp && ParentSetProp->ElementProp != Prop)
+		{
+			ParentSetProp = nullptr;
+		}
+
+		FCachedPropertyPath ResolvedComparisonPath = ResolvedCDOPath;
+		if (ParentArrayProp || ParentMapProp || ParentSetProp)
+		{
+			ResolvedComparisonPath = ResolvedComparisonPath.MakeParentPath();
+		}
+
+		TArray<UObject*> ArchetypeInstances;
+		UnmodifiedCDOObject->GetArchetypeInstances(ArchetypeInstances);
+
+		void* CDOPristineValueAddr = ResolvedComparisonPath.GetCachedAddress();
+
+		FString OldValue;
+		FString NewValue;
+
+		bool bHasCopyableValue = true;
+		bHasCopyableValue &= PropertyPathHelpers::GetPropertyValueAsString(UnmodifiedCDOObject, ChangedPropertyPath, OldValue);
+		bHasCopyableValue &= PropertyPathHelpers::GetPropertyValueAsString(ModifedProxyObject, ChangedPropertyPath, NewValue);
+
+		if (bHasCopyableValue)
+		{
+			FCachedPropertyPath Instance_ComparisonPath = FCachedPropertyPath::MakeUnresolvedCopy(ResolvedComparisonPath);
+			for (UObject* ArchetypeInstance : ArchetypeInstances)
+			{
+				if (ArchetypeInstance != ModifedProxyObject)
+				{
+					Instance_ComparisonPath.Resolve(ArchetypeInstance);
+
+					if (Instance_ComparisonPath.IsFullyResolved())
+					{
+						void* Instance_DestValueAddr = Instance_ComparisonPath.GetCachedAddress();
+						const bool bShouldImport = ResolvedComparisonPath.GetFProperty()->Identical(CDOPristineValueAddr, Instance_DestValueAddr, PPF_DeepComparison);
+
+						if (bShouldImport)
+						{
+							// Note: We use the original change path to apply the actual edit - the comparison path might be shortened due to checking
+							// maps or arrays or sets, but if the array|map|set matches as a whole, then we can mod a specific index.
+							PropertyPathHelpers::SetPropertyValueFromString(ArchetypeInstance, ChangedPropertyPath, NewValue, PPF_InstanceSubobjects);
+						}
+					}
+				}
+			}
+		}
+	}
+}
+
 void FWidgetBlueprintEditor::MigrateFromChain(FEditPropertyChain* PropertyThatChanged, bool bIsModify)
+{
+	MigrateFromChain(nullptr, PropertyThatChanged, bIsModify); 
+}
+
+void FWidgetBlueprintEditor::MigrateFromChain(const FPropertyChangedEvent* PropertyChangedEvent, FEditPropertyChain* PropertyThatChanged, bool bIsModify)
 {
 	UWidgetBlueprint* Blueprint = GetWidgetBlueprintObj();
 
 	UUserWidget* PreviewUserWidget = GetPreview();
 	if ( PreviewUserWidget != nullptr )
 	{
+		// The selected objects is always the root widget, which is a copy of the CDO instanced for preview purposes.
+		// if it's modified we need to copy the values back to the CDO - but we ALSO need to do what the details panel
+		// normally does when modifying a CDO, which is to copy the values to any live instances of the widget with matching
+		// previous values so that defaults stay in-sync with the CDO.
 		for ( TWeakObjectPtr<UObject> ObjectRef : SelectedObjects )
 		{
-			// dealing with root widget here
+			//TODO: We should probably ensure that the objectref is in fact the proxy CDO for the Widget, we should
+			// probably keep track of that somewhere and vet this - you never know if someone might make new changes
+			// that invalidate this assumption.
+
 			FEditPropertyChain::TDoubleLinkedListNode* PropertyChainNode = PropertyThatChanged->GetHead();
 			UObject* WidgetCDO = ObjectRef.Get()->GetClass()->GetDefaultObject(true);
+
+			if (PropertyChangedEvent == nullptr && WidgetCDO)
+			{
+				WidgetCDO->SetEditChangePropagationFlags(EEditChangePropagationFlags::OnlyMarkRealignedInstancesAsDirty);
+				WidgetCDO->PreEditChange(*PropertyThatChanged);
+			}
+
+			if (PropertyChangedEvent && PropertyThatChanged)
+			{
+				// We have to do this before we call MigratePropertyValue, because that will change the CDO, and we won't be
+				// able cheat like we do still having a pristine CDO copy.  The details panel has to do shenanigans where it copies
+				// the data into a temp property piece of memory, but we don't need to do that since we've got the actual unmodified
+				// CDO still - but we're going to propagate the proxy's value that users modified instead of the CDO's value,
+				// then MigratePropertyValue will update the CDO's value.
+				PropagateDefaultPropertyChange(*PropertyChangedEvent, *PropertyThatChanged, WidgetCDO, ObjectRef.Get());
+			}
+
+			// dealing with root widget here
 			MigratePropertyValue(ObjectRef.Get(), WidgetCDO, PropertyChainNode, PropertyChainNode->GetValue(), bIsModify);
+
+			if (PropertyChangedEvent && WidgetCDO)
+			{
+				// 
+				TArray<UObject*> CDOArray = { WidgetCDO };
+				FPropertyChangedEvent ChangeEvent(PropertyChangedEvent->Property, PropertyChangedEvent->ChangeType, CDOArray);
+				ChangeEvent.SetActiveMemberProperty(PropertyChangedEvent->MemberProperty);
+
+				TArray<TMap<FString, int32>> ArrayIndexForProperty;
+				for (TDoubleLinkedList<FProperty*>::TIterator It(PropertyThatChanged->GetHead()); It; ++It)
+				{
+					FProperty* Property = *It;
+
+					const FString PropertyName = Property->GetName();
+					const int32 ArrayIndex = PropertyChangedEvent->GetArrayIndex(PropertyName);
+					if (ArrayIndex != INDEX_NONE)
+					{
+						ArrayIndexForProperty.AddDefaulted();
+						ArrayIndexForProperty.Last().Add(PropertyName, ArrayIndex);
+					}
+				}
+				ChangeEvent.ObjectIteratorIndex = 0;
+				ChangeEvent.SetArrayIndexPerObject(ArrayIndexForProperty);
+
+				FPropertyChangedChainEvent ChainEvent(*PropertyThatChanged, ChangeEvent);
+				WidgetCDO->PostEditChangeChainProperty(ChainEvent);
+				WidgetCDO->SetEditChangePropagationFlags(EEditChangePropagationFlags::None);
+			}
 		}
 
 		for ( FWidgetReference& WidgetRef : SelectedWidgets )
@@ -2577,4 +2722,4 @@ void FWidgetBlueprintEditor::UpdateTrackName(FGuid ObjectId)
 	}
 }
 
-#undef LOCTEXT_NAMESPACE
+#undef LOCTEXT_NAMESPACE
\ No newline at end of file
diff --git a/Engine/Source/Editor/UMGEditor/Public/WidgetBlueprintEditor.h b/Engine/Source/Editor/UMGEditor/Public/WidgetBlueprintEditor.h
--- a/Engine/Source/Editor/UMGEditor/Public/WidgetBlueprintEditor.h
+++ b/Engine/Source/Editor/UMGEditor/Public/WidgetBlueprintEditor.h
@@ -25,6 +25,7 @@ class UWidgetAnimation;
 class UWidgetBlueprint;
 class FPaletteViewModel;
 class FLibraryViewModel;
+struct FPropertyChangedEvent;
 namespace UE::UMG::Editor { class FPreviewMode; }
 
 struct FNamedSlotSelection
@@ -178,9 +179,13 @@ public:
 
 	TSharedPtr<class FWidgetBlueprintEditorToolbar> GetWidgetToolbarBuilder() { return WidgetToolbar; }
 
+	UE_DEPRECATED(5.5, "Use MigrateFromChain with a FPropertyChangedEvent and a FEditPropertyChain")
 	/** Migrate a property change from the preview GUI to the template GUI. */
 	void MigrateFromChain(FEditPropertyChain* PropertyThatChanged, bool bIsModify);
 
+	/** Migrate a property change from the preview GUI to the template GUI. */
+	void MigrateFromChain(const FPropertyChangedEvent* PropertyChangedEvent, FEditPropertyChain* PropertyThatChanged, bool bIsModify);
+
 	/** Event called when an undo/redo transaction occurs */
 	DECLARE_EVENT(FWidgetBlueprintEditor, FOnWidgetBlueprintTransaction)
 	FOnWidgetBlueprintTransaction& GetOnWidgetBlueprintTransaction() { return OnWidgetBlueprintTransaction; }
@@ -543,4 +548,4 @@ private:
 
 	/** When true the external selection is being updated from changes to the sequencer selection. */
 	bool bUpdatingExternalSelection;
-};
+};
\ No newline at end of file
diff --git a/Engine/Source/Runtime/PropertyPath/Private/PropertyPathHelpers.cpp b/Engine/Source/Runtime/PropertyPath/Private/PropertyPathHelpers.cpp
--- a/Engine/Source/Runtime/PropertyPath/Private/PropertyPathHelpers.cpp
+++ b/Engine/Source/Runtime/PropertyPath/Private/PropertyPathHelpers.cpp
@@ -157,7 +157,7 @@ namespace PropertyPathHelpersInternal
 	template<typename ContainerType>
 	struct FCallGetterFunctionAsStringHelper
 	{
-		static bool CallGetterFunction(ContainerType* InContainer, UFunction* InFunction, FString& OutValue) 
+		static bool CallGetterFunction(ContainerType* InContainer, UFunction* InFunction, FString& OutValue, uint32 PPF_Flags)
 		{
 			// Cant call UFunctions on non-UObject containers
 			return false;
@@ -168,7 +168,7 @@ namespace PropertyPathHelpersInternal
 	template<>
 	struct FCallGetterFunctionAsStringHelper<UObject>
 	{
-		static bool CallGetterFunction(UObject* InContainer, UFunction* InFunction, FString& OutValue) 
+		static bool CallGetterFunction(UObject* InContainer, UFunction* InFunction, FString& OutValue, uint32 PPF_Flags)
 		{
 			// We only support calling functions that return a single value and take no parameters.
 			if ( InFunction->NumParms == 1 )
@@ -184,7 +184,7 @@ namespace PropertyPathHelpersInternal
 						ReturnProperty->InitializeValue(TempBuffer.GetData());
 
 						InContainer->ProcessEvent(InFunction, TempBuffer.GetData());
-						ReturnProperty->ExportTextItem_Direct(OutValue, TempBuffer.GetData(), nullptr, nullptr, 0);
+						ReturnProperty->ExportTextItem_Direct(OutValue, TempBuffer.GetData(), nullptr, nullptr, PPF_Flags);
 						return true;
 					}
 				}
@@ -195,7 +195,7 @@ namespace PropertyPathHelpersInternal
 	};
 
 	template<typename ContainerType>
-	static bool GetPropertyValueAsString(ContainerType* InContainer, const FCachedPropertyPath& InPropertyPath, FProperty*& OutProperty, FString& OutValue)
+	static bool GetPropertyValueAsString(ContainerType* InContainer, const FCachedPropertyPath& InPropertyPath, FProperty*& OutProperty, FString& OutValue, uint32 PPF_Flags)
 	{
 		const FPropertyPathSegment& LastSegment = InPropertyPath.GetLastSegment();
 		int32 ArrayIndex = LastSegment.GetArrayIndex();
@@ -213,7 +213,7 @@ namespace PropertyPathHelpersInternal
 				if ( ArrayHelper.IsValidIndex(ArrayIndex) )
 				{
 					OutProperty = ArrayProp->Inner;
-					OutProperty->ExportTextItem_Direct(OutValue, static_cast<void*>(ArrayHelper.GetRawPtr(ArrayIndex)), nullptr, nullptr, 0);
+					OutProperty->ExportTextItem_Direct(OutValue, static_cast<void*>(ArrayHelper.GetRawPtr(ArrayIndex)), nullptr, nullptr, PPF_Flags);
 					return true;
 				}
 			}
@@ -223,14 +223,14 @@ namespace PropertyPathHelpersInternal
 				if ( !!ArrayProp->ContainerPtrToValuePtr<void>(InContainer) )
 				{
 					OutProperty = ArrayProp;
-					OutProperty->ExportTextItem_InContainer(OutValue, InContainer, nullptr, nullptr, 0);
+					OutProperty->ExportTextItem_InContainer(OutValue, InContainer, nullptr, nullptr, PPF_Flags);
 					return true;
 				}
 			}
 		}
 		else if(UFunction* Function = Field.Get<UFunction>())
 		{
-			return FCallGetterFunctionAsStringHelper<ContainerType>::CallGetterFunction(InContainer, Function, OutValue);
+			return FCallGetterFunctionAsStringHelper<ContainerType>::CallGetterFunction(InContainer, Function, OutValue, PPF_Flags);
 		}
 		else if(FProperty* Property = Field.Get<FProperty>())
 		{
@@ -240,7 +240,7 @@ namespace PropertyPathHelpersInternal
 				if ( void* ValuePtr = Property->ContainerPtrToValuePtr<void>(InContainer, ArrayIndex) )
 				{
 					OutProperty = Property;
-					OutProperty->ExportTextItem_Direct(OutValue, ValuePtr, nullptr, nullptr, 0);
+					OutProperty->ExportTextItem_Direct(OutValue, ValuePtr, nullptr, nullptr, PPF_Flags);
 					return true;
 				}
 			}
@@ -253,7 +253,7 @@ namespace PropertyPathHelpersInternal
 	template<typename ContainerType>
 	struct FCallSetterFunctionFromStringHelper
 	{
-		static bool CallSetterFunction(ContainerType* InContainer, UFunction* InFunction, const FString& InValue)
+		static bool CallSetterFunction(ContainerType* InContainer, UFunction* InFunction, const FString& InValue, uint32 PPF_Flags)
 		{
 			// Cant call UFunctions on non-UObject containers
 			return false;
@@ -264,7 +264,7 @@ namespace PropertyPathHelpersInternal
 	template<>
 	struct FCallSetterFunctionFromStringHelper<UObject>
 	{
-		static bool CallSetterFunction(UObject* InContainer, UFunction* InFunction, const FString& InValue)
+		static bool CallSetterFunction(UObject* InContainer, UFunction* InFunction, const FString& InValue, uint32 PPF_Flags)
 		{
 			// We only support calling functions that take a single value and return no parameters
 			if ( InFunction->NumParms == 1 && InFunction->GetReturnProperty() == nullptr )
@@ -279,7 +279,7 @@ namespace PropertyPathHelpersInternal
 						TempBuffer.AddUninitialized(ParamProperty->ElementSize);
 						ParamProperty->InitializeValue(TempBuffer.GetData());
 
-						ParamProperty->ImportText_Direct(*InValue, TempBuffer.GetData(), nullptr, 0);
+						ParamProperty->ImportText_Direct(*InValue, TempBuffer.GetData(), nullptr, PPF_Flags);
 						InContainer->ProcessEvent(InFunction, TempBuffer.GetData());
 						return true;
 					}
@@ -291,7 +291,7 @@ namespace PropertyPathHelpersInternal
 	};
 
 	template<typename ContainerType>
-	static bool SetPropertyValueFromString(ContainerType* InContainer, const FCachedPropertyPath& InPropertyPath, const FString& InValue)
+	static bool SetPropertyValueFromString(ContainerType* InContainer, const FCachedPropertyPath& InPropertyPath, const FString& InValue, uint32 PPF_Flags)
 	{
 		const FPropertyPathSegment& LastSegment = InPropertyPath.GetLastSegment();
 		int32 ArrayIndex = LastSegment.GetArrayIndex();
@@ -308,7 +308,7 @@ namespace PropertyPathHelpersInternal
 				FScriptArrayHelper_InContainer ArrayHelper(ArrayProp, InContainer);
 				if ( ArrayHelper.IsValidIndex(ArrayIndex) )
 				{
-					ArrayProp->Inner->ImportText_Direct(*InValue, static_cast<void*>(ArrayHelper.GetRawPtr(ArrayIndex)), nullptr, 0);
+					ArrayProp->Inner->ImportText_Direct(*InValue, static_cast<void*>(ArrayHelper.GetRawPtr(ArrayIndex)), nullptr, PPF_Flags);
 					return true;
 				}
 			}
@@ -317,14 +317,14 @@ namespace PropertyPathHelpersInternal
 				// No index, so assume we want the array property itself
 				if ( !!ArrayProp->ContainerPtrToValuePtr<void>(InContainer) )
 				{
-					ArrayProp->ImportText_InContainer(*InValue, InContainer, nullptr, 0);
+					ArrayProp->ImportText_InContainer(*InValue, InContainer, nullptr, PPF_Flags);
 					return true;
 				}
 			}
 		}
 		else if(UFunction* Function = Field.Get<UFunction>())
 		{
-			return FCallSetterFunctionFromStringHelper<ContainerType>::CallSetterFunction(InContainer, Function, InValue);
+			return FCallSetterFunctionFromStringHelper<ContainerType>::CallSetterFunction(InContainer, Function, InValue, PPF_Flags);
 		}
 		else if(FProperty* Property = Field.Get<FProperty>())
 		{
@@ -333,7 +333,7 @@ namespace PropertyPathHelpersInternal
 			{
 				if ( void* ValuePtr = Property->ContainerPtrToValuePtr<void>(InContainer, ArrayIndex) )
 				{
-					Property->ImportText_Direct(*InValue, ValuePtr, nullptr, 0);
+					Property->ImportText_Direct(*InValue, ValuePtr, nullptr, PPF_Flags);
 					return true;
 				}
 			}
@@ -783,6 +783,28 @@ FCachedPropertyPath::FCachedPropertyPath(const FPropertyPathSegment& Segment)
 	Segments.Add(Segment);
 }
 
+FCachedPropertyPath::FCachedPropertyPath(const FPropertyChangedChainEvent& PropertyChangedChainEvent)
+	: CachedAddress(nullptr)
+	, CachedFunction(nullptr)
+	, CachedContainer(nullptr)
+	, CachedLastContainerInPath(nullptr)
+	, CachedLastContainerInPathIndex(INDEX_NONE)
+	, bCanSafelyUsedCachedAddress(false)
+{
+	MakeFromChangeEvent(PropertyChangedChainEvent);
+}
+
+FCachedPropertyPath::FCachedPropertyPath(const FPropertyChangedEvent& PropertyChangedEvent, const FEditPropertyChain& PropertyChain)
+	: CachedAddress(nullptr)
+	, CachedFunction(nullptr)
+	, CachedContainer(nullptr)
+	, CachedLastContainerInPath(nullptr)
+	, CachedLastContainerInPathIndex(INDEX_NONE)
+	, bCanSafelyUsedCachedAddress(false)
+{
+	MakeFromChangeEvent(PropertyChangedEvent, PropertyChain);
+}
+
 FCachedPropertyPath::~FCachedPropertyPath() = default;
 
 void FCachedPropertyPath::MakeFromString(const FString& InPropertyPath)
@@ -804,6 +826,26 @@ void FCachedPropertyPath::MakeFromString(const FString& InPropertyPath)
 	Segments.Add(FPropertyPathSegment(Length - Start, &Path[Start]));
 }
 
+void FCachedPropertyPath::MakeFromChangeEvent(const FPropertyChangedChainEvent& PropertyChangedChainEvent)
+{
+	MakeFromChangeEvent(PropertyChangedChainEvent, PropertyChangedChainEvent.PropertyChain);
+}
+
+void FCachedPropertyPath::MakeFromChangeEvent(const FPropertyChangedEvent& PropertyChangedEvent, const FEditPropertyChain& PropertyChain)
+{
+	for (TDoubleLinkedList<FProperty*>::TIterator It(PropertyChain.GetHead()); It; ++It)
+	{
+		FProperty* Property = *It;
+		FString PropertySegment = Property->GetName();
+		const int32 ArrayIndex = PropertyChangedEvent.GetArrayIndex(PropertySegment);
+		if (ArrayIndex != INDEX_NONE)
+		{
+			PropertySegment += FString::Printf(TEXT("[%d]"), ArrayIndex);
+		}
+		Segments.Add(FPropertyPathSegment(PropertySegment.Len(), *PropertySegment));
+	}
+}
+
 FCachedPropertyPath FCachedPropertyPath::MakeUnresolvedCopy(const FCachedPropertyPath& ToCopy)
 {
 	FCachedPropertyPath Path;
@@ -814,6 +856,20 @@ FCachedPropertyPath FCachedPropertyPath::MakeUnresolvedCopy(const FCachedPropert
 	return Path;
 }
 
+FCachedPropertyPath FCachedPropertyPath::MakeParentPath() const
+{
+	// If there are no more segements, then it's just an empty path.
+	if (Segments.Num() == 0)
+	{
+		return *this;
+	}
+
+	FCachedPropertyPath ParentPath = *this;
+	ParentPath.Segments.RemoveAt(ParentPath.Segments.Num() - 1);
+
+	return ParentPath;
+}
+
 int32 FCachedPropertyPath::GetNumSegments() const
 {
 	return Segments.Num();
@@ -1055,121 +1111,125 @@ namespace PropertyPathHelpers
 		}
 	}
 
-	bool GetPropertyValueAsString(UObject* InContainer, const FString& InPropertyPath, FString& OutValue)
+	bool GetPropertyValueAsString(UObject* InContainer, const FString& InPropertyPath, FString& OutValue, uint32 PPF_Flags)
 	{
 		FProperty* Property;
-		return GetPropertyValueAsString(InContainer, InPropertyPath, OutValue, Property);
+		return GetPropertyValueAsString(InContainer, InPropertyPath, OutValue, Property, PPF_Flags);
 	}
 
 	/** Helper for string-based getters */
 	struct FInternalStringGetterResolver : public PropertyPathHelpersInternal::TPropertyPathResolver<FInternalStringGetterResolver>
 	{
-		FInternalStringGetterResolver(FString& InOutValue, FProperty*& InOutProperty)
+		FInternalStringGetterResolver(FString& InOutValue, FProperty*& InOutProperty, uint32 PPF_Flags)
 			: Value(InOutValue)
 			, Property(InOutProperty)
+			, PPF_Flags(PPF_Flags)
 		{
 		}
 
 		template<typename ContainerType>
 		bool Resolve_Impl(ContainerType* InContainer, const FCachedPropertyPath& InPropertyPath)
 		{
-			return PropertyPathHelpersInternal::GetPropertyValueAsString<ContainerType>(InContainer, InPropertyPath, Property, Value);
+			return PropertyPathHelpersInternal::GetPropertyValueAsString<ContainerType>(InContainer, InPropertyPath, Property, Value, PPF_Flags);
 		}
 
 		FString& Value;
 		FProperty*& Property;
+		uint32 PPF_Flags;
 	};
 
-	bool GetPropertyValueAsString(UObject* InContainer, const FString& InPropertyPath, FString& OutValue, FProperty*& OutProperty)
+	bool GetPropertyValueAsString(UObject* InContainer, const FString& InPropertyPath, FString& OutValue, FProperty*& OutProperty, uint32 PPF_Flags)
 	{
 		check(InContainer);
 
-		FInternalStringGetterResolver Resolver(OutValue, OutProperty);
+		FInternalStringGetterResolver Resolver(OutValue, OutProperty, PPF_Flags);
 		return ResolvePropertyPath(InContainer, InPropertyPath, Resolver);
 	}
 
-	bool GetPropertyValueAsString(void* InContainer, UStruct* InStruct, const FString& InPropertyPath, FString& OutValue)
+	bool GetPropertyValueAsString(void* InContainer, UStruct* InStruct, const FString& InPropertyPath, FString& OutValue, uint32 PPF_Flags)
 	{
 		FProperty* Property;
-		return GetPropertyValueAsString(InContainer, InStruct, InPropertyPath, OutValue, Property);
+		return GetPropertyValueAsString(InContainer, InStruct, InPropertyPath, OutValue, Property, PPF_Flags);
 	}
 
-	bool GetPropertyValueAsString(void* InContainer, UStruct* InStruct, const FString& InPropertyPath, FString& OutValue, FProperty*& OutProperty)
+	bool GetPropertyValueAsString(void* InContainer, UStruct* InStruct, const FString& InPropertyPath, FString& OutValue, FProperty*& OutProperty, uint32 PPF_Flags)
 	{
 		check(InContainer);
 		check(InStruct);
 
-		FInternalStringGetterResolver Resolver(OutValue, OutProperty);
+		FInternalStringGetterResolver Resolver(OutValue, OutProperty, PPF_Flags);
 		return ResolvePropertyPath(InContainer, InStruct, InPropertyPath, Resolver);
 	}
 
-	bool GetPropertyValueAsString(UObject* InContainer, const FCachedPropertyPath& InPropertyPath, FString& OutValue)
+	bool GetPropertyValueAsString(UObject* InContainer, const FCachedPropertyPath& InPropertyPath, FString& OutValue, uint32 PPF_Flags)
 	{
 		check(InContainer);
 
 		FProperty* Property;
-		FInternalStringGetterResolver Resolver(OutValue, Property);
+		FInternalStringGetterResolver Resolver(OutValue, Property, PPF_Flags);
 		return ResolvePropertyPath(InContainer, InPropertyPath, Resolver);
 	}
 
-	bool GetPropertyValueAsString(void* InContainer, UStruct* InStruct, const FCachedPropertyPath& InPropertyPath, FString& OutValue)
+	bool GetPropertyValueAsString(void* InContainer, UStruct* InStruct, const FCachedPropertyPath& InPropertyPath, FString& OutValue, uint32 PPF_Flags)
 	{
 		check(InContainer);
 		check(InStruct);
 
 		FProperty* Property;
-		FInternalStringGetterResolver Resolver(OutValue, Property);
+		FInternalStringGetterResolver Resolver(OutValue, Property, PPF_Flags);
 		return ResolvePropertyPath(InContainer, InStruct, InPropertyPath, Resolver);
 	}
 
 	/** Helper for string-based setters */
 	struct FInternalStringSetterResolver : public PropertyPathHelpersInternal::TPropertyPathResolver<FInternalStringSetterResolver>
 	{
-		FInternalStringSetterResolver(const FString& InValueAsString)
+		FInternalStringSetterResolver(const FString& InValueAsString, uint32 PPF_Flags)
 			: Value(InValueAsString)
+			, PPF_Flags(PPF_Flags)
 		{
 		}
 
 		template<typename ContainerType>
 		bool Resolve_Impl(ContainerType* InContainer, const FCachedPropertyPath& InPropertyPath)
 		{
-			return PropertyPathHelpersInternal::SetPropertyValueFromString<ContainerType>(InContainer, InPropertyPath, Value);
+			return PropertyPathHelpersInternal::SetPropertyValueFromString<ContainerType>(InContainer, InPropertyPath, Value, PPF_Flags);
 		}
 	
 		const FString& Value;
+		uint32 PPF_Flags;
 	};
 
-	bool SetPropertyValueFromString(UObject* InContainer, const FString& InPropertyPath, const FString& InValue)
+	bool SetPropertyValueFromString(UObject* InContainer, const FString& InPropertyPath, const FString& InValue, uint32 PPF_Flags)
 	{
 		check(InContainer);
 
-		FInternalStringSetterResolver Resolver(InValue);
+		FInternalStringSetterResolver Resolver(InValue, PPF_Flags);
 		return ResolvePropertyPath(InContainer, InPropertyPath, Resolver);
 	}
 
-	bool SetPropertyValueFromString(UObject* InContainer, const FCachedPropertyPath& InPropertyPath, const FString& InValue)
+	bool SetPropertyValueFromString(UObject* InContainer, const FCachedPropertyPath& InPropertyPath, const FString& InValue, uint32 PPF_Flags)
 	{
 		check(InContainer);
 
-		FInternalStringSetterResolver Resolver(InValue);
+		FInternalStringSetterResolver Resolver(InValue, PPF_Flags);
 		return ResolvePropertyPath(InContainer, InPropertyPath, Resolver);
 	}
 
-	bool SetPropertyValueFromString(void* InContainer, UStruct* InStruct, const FString& InPropertyPath, const FString& InValue)
+	bool SetPropertyValueFromString(void* InContainer, UStruct* InStruct, const FString& InPropertyPath, const FString& InValue, uint32 PPF_Flags)
 	{
 		check(InContainer);
 		check(InStruct);
 
-		FInternalStringSetterResolver Resolver(InValue);
+		FInternalStringSetterResolver Resolver(InValue, PPF_Flags);
 		return ResolvePropertyPath(InContainer, InStruct, InPropertyPath, Resolver);
 	}
 
-	bool SetPropertyValueFromString(void* InContainer, UStruct* InStruct, const FCachedPropertyPath& InPropertyPath, const FString& InValue)
+	bool SetPropertyValueFromString(void* InContainer, UStruct* InStruct, const FCachedPropertyPath& InPropertyPath, const FString& InValue, uint32 PPF_Flags)
 	{
 		check(InContainer);
 		check(InStruct);
 
-		FInternalStringSetterResolver Resolver(InValue);
+		FInternalStringSetterResolver Resolver(InValue, PPF_Flags);
 		return ResolvePropertyPath(InContainer, InStruct, InPropertyPath, Resolver);
 	}
 
@@ -1249,4 +1309,4 @@ namespace PropertyPathHelpers
 		FInternalArrayOperationResolver Resolver(InOperation);
 		return ResolvePropertyPath(InContainer, InPropertyPath, Resolver);
 	}
-}
+}
\ No newline at end of file
diff --git a/Engine/Source/Runtime/PropertyPath/Public/PropertyPathHelpers.h b/Engine/Source/Runtime/PropertyPath/Public/PropertyPathHelpers.h
--- a/Engine/Source/Runtime/PropertyPath/Public/PropertyPathHelpers.h
+++ b/Engine/Source/Runtime/PropertyPath/Public/PropertyPathHelpers.h
@@ -109,6 +109,10 @@ struct FCachedPropertyPath
 	/** */
 	PROPERTYPATH_API FCachedPropertyPath(const TArray<FString>& PathSegments);
 
+	/** Allows constructing a cached property path from the standard events of a PostEditChange Notify or Callback */
+	PROPERTYPATH_API FCachedPropertyPath(const FPropertyChangedChainEvent& PropertyChangedChainEvent);
+	PROPERTYPATH_API FCachedPropertyPath(const FPropertyChangedEvent& PropertyChangedEvent, const FEditPropertyChain& PropertyChain);
+
 	/** */
 	PROPERTYPATH_API ~FCachedPropertyPath();
 
@@ -118,9 +122,16 @@ struct FCachedPropertyPath
 	/** Make a new property path from a string */
 	PROPERTYPATH_API void MakeFromString(const FString& InPropertyPath);
 
+	/** Make a new property path from the chain of properties that come standard in PostEditPropertyChain events */
+	PROPERTYPATH_API void MakeFromChangeEvent(const FPropertyChangedChainEvent& PropertyChangedChainEvent);
+	PROPERTYPATH_API void MakeFromChangeEvent(const FPropertyChangedEvent& PropertyChangedEvent, const FEditPropertyChain& PropertyChain);
+
 	/** Make a copy which is unresolved */
 	static PROPERTYPATH_API FCachedPropertyPath MakeUnresolvedCopy(const FCachedPropertyPath& ToCopy);
 
+	/** Make a copy without the last segment, if the path is resolved, the new path is also resolved. */
+	PROPERTYPATH_API FCachedPropertyPath MakeParentPath() const;
+
 	/** @return Get the number of segments in this path */
 	PROPERTYPATH_API int32 GetNumSegments() const;
 
@@ -261,7 +272,7 @@ namespace PropertyPathHelpers
 	 * @param	OutValue		The string to write the properties value to
 	 * @return true if the property value was successfully copied
 	 */
-	PROPERTYPATH_API bool GetPropertyValueAsString(UObject* InContainer, const FString& InPropertyPath, FString& OutValue);
+	PROPERTYPATH_API bool GetPropertyValueAsString(UObject* InContainer, const FString& InPropertyPath, FString& OutValue, uint32 PPF_Flags = PPF_None);
 
 	/** 
 	 * Get the value represented by this property path as a string 
@@ -271,7 +282,7 @@ namespace PropertyPathHelpers
 	 * @param	OutProperty		The leaf property that the path resolved to
 	 * @return true if the property value was successfully copied
 	 */
-	PROPERTYPATH_API bool GetPropertyValueAsString(UObject* InContainer, const FString& InPropertyPath, FString& OutValue, FProperty*& OutProperty);
+	PROPERTYPATH_API bool GetPropertyValueAsString(UObject* InContainer, const FString& InPropertyPath, FString& OutValue, FProperty*& OutProperty, uint32 PPF_Flags = PPF_None);
 
 	/** 
 	 * Get the value represented by this property path as a string 
@@ -281,7 +292,7 @@ namespace PropertyPathHelpers
 	 * @param	OutValue		The string to write the properties value to
 	 * @return true if the property value was successfully copied
 	 */
-	PROPERTYPATH_API bool GetPropertyValueAsString(void* InContainer, UStruct* InStruct, const FString& InPropertyPath, FString& OutValue);
+	PROPERTYPATH_API bool GetPropertyValueAsString(void* InContainer, UStruct* InStruct, const FString& InPropertyPath, FString& OutValue, uint32 PPF_Flags = PPF_None);
 
 	/** 
 	 * Get the value represented by this property path as a string 
@@ -292,7 +303,7 @@ namespace PropertyPathHelpers
 	 * @param	OutProperty		The leaf property that the path resolved to
 	 * @return true if the property value was successfully copied
 	 */
-	PROPERTYPATH_API bool GetPropertyValueAsString(void* InContainer, UStruct* InStruct, const FString& InPropertyPath, FString& OutValue, FProperty*& OutProperty);
+	PROPERTYPATH_API bool GetPropertyValueAsString(void* InContainer, UStruct* InStruct, const FString& InPropertyPath, FString& OutValue, FProperty*& OutProperty, uint32 PPF_Flags = PPF_None);
 
 	/** 
 	 * Get the value represented by this property path as a string 
@@ -301,7 +312,7 @@ namespace PropertyPathHelpers
 	 * @param	OutValue		The string to write the properties value to
 	 * @return true if the property value was successfully copied
 	 */
-	PROPERTYPATH_API bool GetPropertyValueAsString(UObject* InContainer, const FCachedPropertyPath& InPropertyPath, FString& OutValue);
+	PROPERTYPATH_API bool GetPropertyValueAsString(UObject* InContainer, const FCachedPropertyPath& InPropertyPath, FString& OutValue, uint32 PPF_Flags = PPF_None);
 
 	/** 
 	 * Get the value represented by this property path as a string 
@@ -311,7 +322,7 @@ namespace PropertyPathHelpers
 	 * @param	OutValue		The string to write the properties value to
 	 * @return true if the property value was successfully copied
 	 */
-	PROPERTYPATH_API bool GetPropertyValueAsString(void* InContainer, UStruct* InStruct, const FCachedPropertyPath& InPropertyPath, FString& OutValue);
+	PROPERTYPATH_API bool GetPropertyValueAsString(void* InContainer, UStruct* InStruct, const FCachedPropertyPath& InPropertyPath, FString& OutValue, uint32 PPF_Flags = PPF_None);
 
 	/** 
 	 * Set the value represented by this property path from a string 
@@ -320,7 +331,7 @@ namespace PropertyPathHelpers
 	 * @param	InValue			The string to read the properties value from
 	 * @return true if the property value was successfully copied
 	 */
-	PROPERTYPATH_API bool SetPropertyValueFromString(UObject* InContainer, const FString& InPropertyPath, const FString& InValue);
+	PROPERTYPATH_API bool SetPropertyValueFromString(UObject* InContainer, const FString& InPropertyPath, const FString& InValue, uint32 PPF_Flags = PPF_None);
 
 	/** 
 	 * Set the value represented by this property path from a string 
@@ -329,7 +340,7 @@ namespace PropertyPathHelpers
 	 * @param	InValue			The string to read the properties value from
 	 * @return true if the property value was successfully copied
 	 */
-	PROPERTYPATH_API bool SetPropertyValueFromString(UObject* InContainer, const FCachedPropertyPath& InPropertyPath, const FString& InValue);
+	PROPERTYPATH_API bool SetPropertyValueFromString(UObject* InContainer, const FCachedPropertyPath& InPropertyPath, const FString& InValue, uint32 PPF_Flags = PPF_None);
 
 	/** 
 	 * Set the value represented by this property path from a string 
@@ -339,7 +350,7 @@ namespace PropertyPathHelpers
 	 * @param	InValue			The string to read the properties value from
 	 * @return true if the property value was successfully copied
 	 */
-	PROPERTYPATH_API bool SetPropertyValueFromString(void* InContainer, UStruct* InStruct, const FString& InPropertyPath, const FString& InValue);
+	PROPERTYPATH_API bool SetPropertyValueFromString(void* InContainer, UStruct* InStruct, const FString& InPropertyPath, const FString& InValue, uint32 PPF_Flags = PPF_None);
 
 	/** 
 	 * Set the value represented by this property path from a string 
@@ -349,7 +360,7 @@ namespace PropertyPathHelpers
 	 * @param	InValue			The string to read the properties value from
 	 * @return true if the property value was successfully copied
 	 */
-	PROPERTYPATH_API bool SetPropertyValueFromString(void* InContainer, UStruct* InStruct, const FCachedPropertyPath& InPropertyPath, const FString& InValue);
+	PROPERTYPATH_API bool SetPropertyValueFromString(void* InContainer, UStruct* InStruct, const FCachedPropertyPath& InPropertyPath, const FString& InValue, uint32 PPF_Flags = PPF_None);
 
 	/** 
 	 * Get the value represented by this property path. forcing the use of cached addresses whether or
@@ -1440,4 +1451,4 @@ namespace PropertyPathHelpersInternal
 
 		return false;
 	}
-}
+}
\ No newline at end of file
